Visualizer fixes for TWeakObjectPtr due to GUObjectArray changes.
TWeakObjectPtr<T> now visualizes T instead of UObject.
PS4 visualization for TFixedSizeArrayView synced with VC implementation.

diff --git a/Engine/Extras/LLDBDataFormatters/UE4DataFormatters.py b/Engine/Extras/LLDBDataFormatters/UE4DataFormatters.py
--- a/Engine/Extras/LLDBDataFormatters/UE4DataFormatters.py
+++ b/Engine/Extras/LLDBDataFormatters/UE4DataFormatters.py
@@ -108,17 +108,17 @@ class UE4TWeakObjectPtrSynthProvider:
         logger = lldb.formatters.Logger.Logger()
         logger >> "Retrieving child " + str(index)
         if self.ObjectSerialNumberVal >= 1:
-            Expr = '*((*(*GSerialNumberBlocksForDebugVisualizers + ('+str(self.ObjectIndexVal)+' / 0x4000))) + ('+str(self.ObjectIndexVal)+' % 0x4000)) == '+str(self.ObjectSerialNumberVal)
+            Expr = 'GObjectArrayForDebugVisualizers->Objects['+str(self.ObjectIndexVal)+'].SerialNumber == '+str(self.ObjectSerialNumberVal)
             Val = self.valobj.CreateValueFromExpression(str(self.ObjectIndexVal), Expr)
             Value = Val.GetValueAsUnsigned(0)
             if Value != 0:
-                Expr = 'GObjectArrayForDebugVisualizers['+str(self.ObjectIndexVal)+' / 0x4000]['+str(self.ObjectIndexVal)+' % 0x4000]'
+                Expr = 'GObjectArrayForDebugVisualizers->Objects['+str(self.ObjectIndexVal)+'].Object'
                 return self.valobj.CreateValueFromExpression('Object', Expr)
             else:
                 Expr = '(void*)0xDEADBEEF'
                 return self.valobj.CreateValueFromExpression('Object', Expr)
 
-        Expr = '(void*)NULL'
+        Expr = 'nullptr'
         return self.valobj.CreateValueFromExpression('Object', Expr)
     
     def update(self):
@@ -138,16 +138,16 @@ def UE4FWeakObjectPtrSummaryProvider(valobj,dict):
     ObjectSerialNumber = valobj.GetChildMemberWithName('ObjectSerialNumber')
     ObjectSerialNumberVal = ObjectSerialNumber.GetValueAsSigned(0)
     if ObjectSerialNumberVal < 1:
-        return 'object=NULL'
+        return 'object=nullptr'
     ObjectIndex = valobj.GetChildMemberWithName('ObjectIndex')
     ObjectIndexVal = ObjectIndex.GetValueAsSigned(0)
-    Expr = '*((*(*GSerialNumberBlocksForDebugVisualizers + ('+str(ObjectIndexVal)+' / 0x4000))) + ('+str(ObjectIndexVal)+' % 0x4000)) == '+str(ObjectSerialNumberVal)
+    Expr = 'GObjectArrayForDebugVisualizers->Objects['+str(ObjectIndexVal)+'].SerialNumber == '+str(ObjectSerialNumberVal)
     Val = valobj.CreateValueFromExpression(str(ObjectIndexVal), Expr)
     ValRef = Val.GetValueAsUnsigned(0)
     if ValRef == 0:
         return 'object=STALE'
     else:
-        Expr = '(UObjectBase*)GObjectArrayForDebugVisualizers['+str(ObjectIndexVal)+' / 0x4000]['+str(ObjectIndexVal)+' % 0x4000]'
+        Expr = 'GObjectArrayForDebugVisualizers->Objects['+str(ObjectIndexVal)+'].Object'
         Val = valobj.CreateValueFromExpression(str(ObjectIndexVal), Expr)
         return 'object=' + Val.GetValue()
 
