Android 2-byte string support
#ue4
#android

diff --git a/Engine/Extras/LLDBDataFormatters/UE4DataFormatters_2ByteChars.py b/Engine/Extras/LLDBDataFormatters/UE4DataFormatters_2ByteChars.py
new file mode 100644
--- /dev/null
+++ b/Engine/Extras/LLDBDataFormatters/UE4DataFormatters_2ByteChars.py
@@ -0,0 +1,666 @@
+# Copyright 1998-2018 Epic Games, Inc. All Rights Reserved.
+
+# /*=============================================================================
+#	LLDB Data Formatters for Unreal Types
+# =============================================================================*/
+
+import lldb
+import lldb.formatters.Logger
+
+# Uncomment the line below to have the data formatters emit debug logging
+# lldb.formatters.Logger._lldb_formatters_debug_level=1
+
+# What documentation there is for parsing values in LLDB can be found here: http://lldb.llvm.org/cpp_reference/html/classlldb_1_1SBValue.html
+
+# To install:
+# 1) Open Terminal and run:
+#        touch ~/.lldbinit
+#        open ~/.lldbinit
+# 2) Add the following text to .lldbini and save - modifying the path as appropriate:
+#        settings set target.inline-breakpoint-strategy always
+#        command script import "/Path/To/Epic/UE4/Engine/Extras/LLDBDataFormatters/UE4DataFormatters.py"
+# 3) Restart Xcode
+
+def UE4TCharSummaryProvider(valobj,dict):
+    Data = valobj.GetValue()
+    Val = valobj.GetSummary()
+    Type = valobj.GetType().GetUnqualifiedType()
+    if Type.IsPointerType():
+        DataVal = valobj.GetValueAsUnsigned(0)
+        if DataVal == 0:
+            Val = 'NULL'
+        else:
+            Expr = '(ushort_t*)('+str(Data)+')'
+            ValRef = valobj.CreateValueFromExpression('string', Expr)
+            Val = str(ValRef.GetSummary())
+    elif Type.IsReferenceType():
+        Expr = '(ushort_t&)('+str(Data)+')'
+        ValRef = valobj.CreateValueFromExpression('string', Expr)
+        Val = str(ValRef.GetSummary())
+    return Val
+
+def UE4FStringSummaryProvider(valobj,dict):
+    Data = valobj.GetChildMemberWithName('Data')
+    ArrayNumVal = Data.GetNumChildren()
+    if ArrayNumVal < 0:
+       return 'string=Invalid'
+    elif ArrayNumVal == 0:
+       return 'string=Empty'
+    else:
+        AllocatorInstance = Data.GetChildMemberWithName('AllocatorInstance')
+        ActualData = AllocatorInstance.GetChildMemberWithName('Data')
+        Expr = '(ushort_t*)('+str(ActualData.GetValue())+')'
+        ValRef = valobj.CreateValueFromExpression('string', Expr)
+        Val = str(ValRef.GetSummary())
+        return 'string=' + Val
+
+def UE4FNameSummaryProvider(valobj,dict):
+    Index = valobj.GetChildMemberWithName('DisplayIndex')
+    if not Index.IsValid():
+        Index = valobj.GetChildMemberWithName('ComparisonIndex')
+    IndexVal = Index.GetValueAsUnsigned(0)
+    if IndexVal >= 1048576:
+        return 'name=Invalid'
+    else:
+        Expr = '(char*)(((FNameEntry***)FName::GetNameTableForDebuggerVisualizers_MT())['+str(IndexVal)+' / 16384]['+str(IndexVal)+' % 16384])->AnsiName'
+        FNameRef = valobj.CreateValueFromExpression(str(IndexVal), Expr)
+        assert FNameRef != None
+        Val = FNameRef.GetSummary()
+        if IndexVal == 0:
+            return 'name=' + Val + '_' + str(IndexVal-1)
+        else:
+            return 'name=' + Val
+			
+def UE4FMinimalNameSummaryProvider(valobj,dict):
+    Index = valobj.GetChildMemberWithName('Index')
+    IndexVal = Index.GetValueAsUnsigned(0)
+    if IndexVal >= 1048576:
+        return 'name=Invalid'
+    else:
+        Expr = '(char*)(((FNameEntry***)FName::GetNameTableForDebuggerVisualizers_MT())['+str(IndexVal)+' / 16384]['+str(IndexVal)+' % 16384])->AnsiName'
+        FNameRef = valobj.CreateValueFromExpression(str(IndexVal), Expr)
+        assert FNameRef != None
+        Val = FNameRef.GetSummary()
+        if IndexVal == 0:
+            return 'name=' + Val + '_' + str(IndexVal-1)
+        else:
+            return 'name=' + Val
+
+def UE4UObjectBaseSummaryProvider(valobj,dict):
+    Name = valobj.GetChildMemberWithName('Name')
+    return Name.GetSummary()
+
+class UE4TWeakObjectPtrSynthProvider:
+    
+    def __init__(self, valobj, dict):
+        logger = lldb.formatters.Logger.Logger()
+        self.valobj = valobj
+    
+    def num_children(self):
+        logger = lldb.formatters.Logger.Logger()
+        return 1
+    
+    def get_child_index(self,name):
+        logger = lldb.formatters.Logger.Logger()
+        return 0
+    
+    def get_child_at_index(self,index):
+        logger = lldb.formatters.Logger.Logger()
+        logger >> "Retrieving child " + str(index)
+        if self.ObjectSerialNumberVal >= 1:
+            Expr = 'GObjectArrayForDebugVisualizers->Objects['+str(self.ObjectIndexVal)+'].SerialNumber == '+str(self.ObjectSerialNumberVal)
+            Val = self.valobj.CreateValueFromExpression(str(self.ObjectIndexVal), Expr)
+            Value = Val.GetValueAsUnsigned(0)
+            if Value != 0:
+                Expr = 'GObjectArrayForDebugVisualizers->Objects['+str(self.ObjectIndexVal)+'].Object'
+                return self.valobj.CreateValueFromExpression('Object', Expr)
+            else:
+                Expr = '(void*)0xDEADBEEF'
+                return self.valobj.CreateValueFromExpression('Object', Expr)
+
+        Expr = 'nullptr'
+        return self.valobj.CreateValueFromExpression('Object', Expr)
+    
+    def update(self):
+        logger = lldb.formatters.Logger.Logger()
+        try:
+            self.ObjectSerialNumber = self.valobj.GetChildMemberWithName('ObjectSerialNumber')
+            self.ObjectSerialNumberVal = self.ObjectSerialNumber.GetValueAsSigned(0)
+            self.ObjectIndex = self.valobj.GetChildMemberWithName('ObjectIndex')
+            self.ObjectIndexVal = self.ObjectIndex.GetValueAsSigned(0)
+        except:
+            pass
+
+    def has_children(self):
+            return True
+
+def UE4FWeakObjectPtrSummaryProvider(valobj,dict):
+    ObjectSerialNumber = valobj.GetChildMemberWithName('ObjectSerialNumber')
+    ObjectSerialNumberVal = ObjectSerialNumber.GetValueAsSigned(0)
+    if ObjectSerialNumberVal < 1:
+        return 'object=nullptr'
+    ObjectIndex = valobj.GetChildMemberWithName('ObjectIndex')
+    ObjectIndexVal = ObjectIndex.GetValueAsSigned(0)
+    Expr = 'GObjectArrayForDebugVisualizers->Objects['+str(ObjectIndexVal)+'].SerialNumber == '+str(ObjectSerialNumberVal)
+    Val = valobj.CreateValueFromExpression(str(ObjectIndexVal), Expr)
+    ValRef = Val.GetValueAsUnsigned(0)
+    if ValRef == 0:
+        return 'object=STALE'
+    else:
+        Expr = 'GObjectArrayForDebugVisualizers->Objects['+str(ObjectIndexVal)+'].Object'
+        Val = valobj.CreateValueFromExpression(str(ObjectIndexVal), Expr)
+        return 'object=' + Val.GetValue()
+
+class UE4ChunkedArraySynthProvider:
+    
+    def __init__(self, valobj, dict):
+        logger = lldb.formatters.Logger.Logger()
+        self.valobj = valobj
+    
+    def num_children(self):
+        logger = lldb.formatters.Logger.Logger()
+        try:
+            NumElementsVal = self.NumElements.GetValueAsSigned(0)
+            return NumElementsVal;
+        except:
+            return 0;
+    
+    def get_child_index(self,name):
+        logger = lldb.formatters.Logger.Logger()
+        try:
+            return int(name.lstrip('[').rstrip(']'))
+        except:
+            return None
+    
+    def get_child_at_index(self,index):
+        logger = lldb.formatters.Logger.Logger()
+        logger >> "Retrieving child " + str(index)
+        Expr = '(unsigned)sizeof('+str(self.valobj.GetType().GetUnqualifiedType().GetName())+'::FChunk)/'+str(self.ElementTypeSize)
+        self.ChunkBytes = self.valobj.CreateValueFromExpression('['+str(index)+']', Expr)
+        self.ChunkBytesSize = self.ChunkBytes.GetValueAsUnsigned(0)
+        assert self.ChunkBytesSize != 0
+        
+        Expr = '*(*((('+str(self.ElementType.GetName())+'**)'+str(self.AllocatorData.GetValue())+')+'+str(index / self.ChunkBytesSize)+')+'+str(index % self.ChunkBytesSize)+')'
+        return self.valobj.CreateValueFromExpression('['+str(index)+']', Expr)
+    
+    def extract_type(self):
+        logger = lldb.formatters.Logger.Logger()
+        ArrayType = self.valobj.GetType().GetUnqualifiedType()
+        if ArrayType.IsReferenceType():
+            ArrayType = ArrayType.GetDereferencedType()
+        elif ArrayType.IsPointerType():
+            ArrayType = ArrayType.GetPointeeType()
+        if ArrayType.GetNumberOfTemplateArguments() > 0:
+            ElementType = ArrayType.GetTemplateArgumentType(0)
+        else:
+            ElementType = None
+        return ElementType
+    
+    def update(self):
+        logger = lldb.formatters.Logger.Logger()
+        try:
+            self.ElementType = self.extract_type()
+            self.ElementTypeSize = self.ElementType.GetByteSize()
+            assert self.ElementTypeSize != 0
+            self.NumElements = self.valobj.GetChildMemberWithName('NumElements')
+            self.Chunks = self.valobj.GetChildMemberWithName('Chunks')
+            self.ArrayNum = self.Chunks.GetChildMemberWithName('ArrayNum')
+            self.AllocatorInstance = self.Chunks.GetChildMemberWithName('AllocatorInstance')
+            self.AllocatorData = self.AllocatorInstance.GetChildMemberWithName('Data')
+        except:
+            pass
+    
+    def has_children(self):
+        return True
+
+def UE4ChunkedArraySummaryProvider(valobj,dict):
+    return 'size=' + str(valobj.GetNumChildren())
+
+class UE4SparseArraySynthProvider:
+    
+    def __init__(self, valobj, dict):
+        logger = lldb.formatters.Logger.Logger()
+        self.valobj = valobj
+    
+    def num_children(self):
+        logger = lldb.formatters.Logger.Logger()
+        try:
+            NumBitsVal = self.NumFreeIndices.GetValueAsSigned(0)
+            ArrayNumVal = self.Data.GetChildMemberWithName('ArrayNum').GetValueAsSigned(0)
+            return ArrayNumVal - NumBitsVal;
+        except:
+            return 0;
+    
+    def get_child_index(self,name):
+        logger = lldb.formatters.Logger.Logger()
+        try:
+            return int(name.lstrip('[').rstrip(']'))
+        except:
+            return None
+    
+    def get_child_at_index(self,index):
+        logger = lldb.formatters.Logger.Logger()
+        logger >> "Retrieving child " + str(index)
+        if index < 0:
+            return None;
+        
+        if index >= self.num_children():
+            return None;
+        
+        Val = None
+        if self.SecondaryDataDataVal > 0:
+            Expr = '(bool)((((int*)'+str(self.SecondaryDataData.GetAddress())+')['+str(index)+'/32] >> '+str(index)+') & 1)'
+            Val = self.SecondaryDataData.CreateValueFromExpression('['+str(index)+']', Expr)
+        else:
+            Expr = '(bool)((((int*)'+str(self.InlineData.GetAddress())+')['+str(index)+'/32] >> '+str(index)+') & 1)'
+            Val = self.InlineData.CreateValueFromExpression('['+str(index)+']', Expr)
+
+        if Val.GetValueAsSigned(0) != 0:
+            offset = index * self.ElementTypeSize
+            return self.AllocatorData.CreateChildAtOffset('['+str(index)+']',offset,self.ElementType)
+        else:
+            return None
+                
+    def extract_type(self):
+        logger = lldb.formatters.Logger.Logger()
+        ArrayType = self.valobj.GetType().GetUnqualifiedType()
+        if ArrayType.IsReferenceType():
+            ArrayType = ArrayType.GetDereferencedType()
+        elif ArrayType.IsPointerType():
+            ArrayType = ArrayType.GetPointeeType()
+        if ArrayType.GetNumberOfTemplateArguments() > 0:
+            ElementType = ArrayType.GetTemplateArgumentType(0)
+        else:
+            ElementType = None
+        return ElementType
+    
+    def update(self):
+        logger = lldb.formatters.Logger.Logger()
+        try:
+            self.ElementType = self.extract_type()
+            self.ElementTypeSize = self.ElementType.GetByteSize()
+            assert self.ElementTypeSize != 0
+            self.NumFreeIndices = self.valobj.GetChildMemberWithName('NumFreeIndices')
+            self.Data = self.valobj.GetChildMemberWithName('Data')
+            self.AllocatorInstance = self.Data.GetChildMemberWithName('AllocatorInstance')
+            self.AllocatorData = self.AllocatorInstance.GetChildMemberWithName('Data')
+            self.AllocationFlags = self.valobj.GetChildMemberWithName('AllocationFlags')
+            self.InlineData = self.AllocationFlags.GetChildMemberWithName('InlineData')
+            self.SecondaryData = self.AllocationFlags.GetChildMemberWithName('SecondaryData')
+            self.SecondaryDataData = self.SecondaryData.GetChildMemberWithName('Data')
+            self.SecondaryDataDataVal = self.SecondaryDataData.GetValueAsSigned(0)
+        except:
+            pass
+    
+    def has_children(self):
+        return True
+
+def UE4SparseArraySummaryProvider(valobj,dict):
+    return 'size=' + str(valobj.GetNumChildren())
+
+class UE4BitArraySynthProvider:
+   
+    def __init__(self, valobj, dict):
+        logger = lldb.formatters.Logger.Logger()
+        self.valobj = valobj
+    
+    def num_children(self):
+        logger = lldb.formatters.Logger.Logger()
+        try:
+            NumBitsVal = self.NumBits.GetValueAsSigned(0)
+            return NumBitsVal;
+        except:
+            return 0;
+
+    def get_child_index(self,name):
+        logger = lldb.formatters.Logger.Logger()
+        try:
+            return int(name.lstrip('[').rstrip(']'))
+        except:
+            return None
+
+    def get_child_at_index(self,index):
+        logger = lldb.formatters.Logger.Logger()
+        logger >> "Retrieving child " + str(index)
+        if index < 0:
+            return None;
+
+        if index >= self.num_children():
+            return None;
+
+        if self.SecondaryDataDataVal > 0:
+            Expr = '(bool)((((int*)'+str(self.SecondaryDataData.GetAddress())+')['+str(index)+'/32] >> '+str(index)+') & 1)'
+            return self.SecondaryDataData.CreateValueFromExpression('['+str(index)+']', Expr)
+        else:
+            Expr = '(bool)((((int*)'+str(self.InlineData.GetAddress())+')['+str(index)+'/32] >> '+str(index)+') & 1)'
+            return self.InlineData.CreateValueFromExpression('['+str(index)+']', Expr)
+
+        return None
+
+    def update(self):
+        logger = lldb.formatters.Logger.Logger()
+        try:
+            self.SecondaryDataData = None
+            self.SecondaryDataDataVal = 0
+            self.NumBits = self.valobj.GetChildMemberWithName('NumBits')
+            self.MaxBits = self.valobj.GetChildMemberWithName('MaxBits')
+            self.InlineData = self.valobj.GetChildMemberWithName('InlineData')
+            self.SecondaryData = self.valobj.GetChildMemberWithName('SecondaryData')
+            if self.SecondaryData != None:
+                self.SecondaryDataData = self.SecondaryData.GetChildMemberWithName('Data')
+                self.SecondaryDataDataVal = self.SecondaryDataData.GetValueAsUnsigned(0)
+        except:
+            pass
+
+    def has_children(self):
+        return True
+
+def UE4BitArraySummaryProvider(valobj,dict):
+    return 'size=' + str(valobj.GetNumChildren())
+
+class UE4ArraySynthProvider:
+    
+    def __init__(self, valobj, dict):
+        logger = lldb.formatters.Logger.Logger()
+        self.valobj = valobj
+    
+    def num_children(self):
+        logger = lldb.formatters.Logger.Logger()
+        try:
+            ArrayNumVal = self.ArrayNum.GetValueAsSigned(0)
+            return ArrayNumVal + self.NumChildren;
+        except:
+            return 0;
+
+    def get_child_index(self,name):
+        logger = lldb.formatters.Logger.Logger()
+        try:
+            return self.NumChildren + int(name.lstrip('[').rstrip(']'))
+        except:
+            return self.valobj.GetIndexOfChildWithName(name)
+    
+    def get_child_at_index(self,index):
+        logger = lldb.formatters.Logger.Logger()
+        logger >> "Retrieving child " + str(index)
+        if index < 0:
+            return None;
+        
+        if index < self.NumChildren:
+            return self.valobj.GetChildAtIndex(index)
+        else:
+            index -= self.NumChildren
+
+        if index >= self.num_children():
+            return None;
+        try:
+            offset = index * self.ElementTypeSize
+            if self.Data != None:
+                return self.Data.CreateChildAtOffset('['+str(index)+']',offset,self.ElementType)
+            elif self.SecondaryDataDataVal > 0:
+                return self.SecondaryDataData.CreateChildAtOffset('['+str(index)+']',offset,self.ElementType)
+            else:
+                return self.InlineData.CreateChildAtOffset('['+str(index)+']',offset,self.ElementType)
+        except:
+            return None
+
+    def extract_type(self):
+        logger = lldb.formatters.Logger.Logger()
+        ArrayType = self.valobj.GetType().GetUnqualifiedType()
+        if ArrayType.IsReferenceType():
+            ArrayType = ArrayType.GetDereferencedType()
+        elif ArrayType.IsPointerType():
+            ArrayType = ArrayType.GetPointeeType()
+        
+        if ArrayType.GetNumberOfTemplateArguments() > 0:
+            ElementType = ArrayType.GetTemplateArgumentType(0)
+        else:
+            ElementType = None
+        return ElementType
+    
+    def update(self):
+        logger = lldb.formatters.Logger.Logger()
+        try:
+            self.NumChildren = self.valobj.GetNumChildren()
+            self.ArrayNum = self.valobj.GetChildMemberWithName('ArrayNum')
+            self.ArrayMax = self.valobj.GetChildMemberWithName('ArrayMax')
+            self.AllocatorInstance = self.valobj.GetChildMemberWithName('AllocatorInstance')
+            if self.AllocatorInstance.GetType().IsReferenceType():
+                self.AllocatorInstance = self.AllocatorInstance.Dereference()
+            self.Data = self.AllocatorInstance.GetChildMemberWithName('Data')
+            self.InlineData = self.AllocatorInstance.GetChildMemberWithName('InlineData')
+            self.SecondaryData = self.AllocatorInstance.GetChildMemberWithName('SecondaryData')
+            if self.SecondaryData != None:
+                self.SecondaryDataData = self.SecondaryData.GetChildMemberWithName('Data')
+                self.SecondaryDataDataVal = self.SecondaryDataData.GetValueAsUnsigned(0)
+            else:
+                self.SecondaryDataData = None
+                self.SecondaryDataDataVal = 0
+        except:
+            logger >> "UE4ArraySynthProvider::update failed accessing members"
+            pass
+        try:
+            self.ElementType = self.extract_type()
+            self.ElementTypeSize = self.ElementType.GetByteSize()
+            assert self.ElementTypeSize != 0
+        except:
+            logger >> "UE4ArraySynthProvider::update failed accessing element type"
+            pass
+
+    def has_children(self):
+        return True
+
+def UE4ArraySummaryProvider(valobj,dict):
+    return 'size=' + str(valobj.GetChildMemberWithName('ArrayNum').GetValueAsSigned(0))
+
+class UE4SetSynthProvider:
+    
+    def __init__(self, valobj, dict):
+        logger = lldb.formatters.Logger.Logger()
+        self.valobj = valobj
+    
+    def num_children(self):
+        logger = lldb.formatters.Logger.Logger()
+        try:
+            self.ArrayNumVal = self.ArrayNum.GetValueAsUnsigned(0)
+            self.NumFreeIndicesVal = self.NumFreeIndices.GetValueAsUnsigned(0)
+            return self.ArrayNumVal - self.NumFreeIndicesVal;
+        except:
+            return 0;
+    
+    def get_child_index(self,name):
+        logger = lldb.formatters.Logger.Logger()
+        try:
+            return int(name.lstrip('[').rstrip(']'))
+        except:
+            return 0
+    
+    def get_child_at_index(self,index):
+        logger = lldb.formatters.Logger.Logger()
+        logger >> "Retrieving child " + str(index)
+        if index < 0:
+            return None;
+        
+        if index >= self.num_children():
+            return None;
+        try:
+            offset = index * self.ElementTypeSize
+            HasObject = 0
+            if self.SecondaryDataDataVal > 0:
+                Expr = '(bool)((((int*)'+str(self.SecondaryDataDataVal)+')['+str(index)+'/32] >> '+str(index)+') & 1)'
+                HasObject = 1 ##self.AllocationFlagsSecondaryDataData.CreateValueFromExpression('['+str(index)+']', Expr).GetValueAsUnsigned(0)
+            else:
+                Expr = '(bool)((((int*)'+str(self.AllocationFlagsInlineDataAddr)+')['+str(index)+'/32] >> '+str(index)+') & 1)'
+                HasObject = 1 ##self.AllocationFlagsInlineData.CreateValueFromExpression('['+str(index)+']', Expr).GetValueAsUnsigned(0)
+
+            if HasObject == 1:
+                return self.AllocatorInstanceData.CreateChildAtOffset('['+str(index)+']',offset,self.ElementType)
+            else:
+                return self.valobj.CreateValueFromExpression('['+str(index)+']', '(void*)0xDEADBEEF')
+        except:
+            return None
+    
+    def extract_type(self):
+        logger = lldb.formatters.Logger.Logger()
+        ArrayType = self.valobj.GetType().GetUnqualifiedType()
+        if ArrayType.IsReferenceType():
+            ArrayType = ArrayType.GetDereferencedType()
+        elif ArrayType.IsPointerType():
+            ArrayType = ArrayType.GetPointeeType()
+        if ArrayType.GetNumberOfTemplateArguments() > 0:
+            ElementType = ArrayType.GetTemplateArgumentType(0)
+        else:
+            ElementType = None
+        return ElementType
+    
+    def update(self):
+        logger = lldb.formatters.Logger.Logger()
+        try:
+            self.Elements = self.valobj.GetChildMemberWithName('Elements')
+            self.ElementsData = self.Elements.GetChildMemberWithName('Data')
+            self.ArrayNum = self.ElementsData.GetChildMemberWithName('ArrayNum')
+            self.NumFreeIndices = self.Elements.GetChildMemberWithName('NumFreeIndices')
+            self.AllocationFlags = self.Elements.GetChildMemberWithName('AllocationFlags')
+            self.AllocationFlagsAllocatorInstance = self.AllocationFlags.GetChildMemberWithName('AllocatorInstance')
+            self.AllocatorInstance = self.ElementsData.GetChildMemberWithName('AllocatorInstance')
+            self.AllocatorInstanceData = self.AllocatorInstance.GetChildMemberWithName('Data')
+            self.AllocationFlagsInlineData = self.AllocationFlagsAllocatorInstance.GetChildMemberWithName('InlineData')
+            self.AllocationFlagsInlineDataAddr = self.AllocationFlagsInlineData.AddressOf().GetValueAsUnsigned(0)
+            self.AllocationFlagsSecondaryData = self.AllocationFlagsAllocatorInstance.GetChildMemberWithName('SecondaryData')
+            self.AllocationFlagsSecondaryDataData = self.AllocationFlagsSecondaryData.GetChildMemberWithName('Data')
+            self.SecondaryDataDataVal = self.AllocationFlagsSecondaryDataData.GetValueAsUnsigned(0)
+            self.ElementType = self.extract_type()
+            # Can't cast to TSetElement - the template instantiation won't allow it
+            Expr = '(size_t)sizeof(FSetElementId) + sizeof(int32)'
+            TSetElement = self.valobj.CreateValueFromExpression('TSetElement', Expr)
+            # This may fail due to C++ struct padding - will have to check
+            self.ElementTypeSize = self.ElementType.GetByteSize() + TSetElement.GetValueAsUnsigned(0)
+            assert self.ElementTypeSize != 0
+        except:
+            pass
+    
+    def has_children(self):
+        return True
+
+def UE4SetSummaryProvider(valobj,dict):
+    return 'size=' + str(valobj.GetNumChildren())
+
+class UE4MapSynthProvider:
+    
+    def __init__(self, valobj, dict):
+        logger = lldb.formatters.Logger.Logger()
+        self.valobj = valobj
+    
+    def num_children(self):
+        logger = lldb.formatters.Logger.Logger()
+        try:
+            self.ArrayNumVal = self.ArrayNum.GetValueAsUnsigned(0)
+            self.NumFreeIndicesVal = self.NumFreeIndices.GetValueAsUnsigned(0)
+            return self.ArrayNumVal - self.NumFreeIndicesVal;
+        except:
+            return 0;
+    
+    def get_child_index(self,name):
+        logger = lldb.formatters.Logger.Logger()
+        try:
+            return int(name.lstrip('[').rstrip(']'))
+        except:
+            return 0
+    
+    def get_child_at_index(self,index):
+        logger = lldb.formatters.Logger.Logger()
+        logger >> "Retrieving child " + str(index)
+        if index < 0:
+            return None;
+        
+        if index >= self.num_children():
+            return None;
+        try:
+            offset = index * self.ElementTypeSize
+            HasObject = 0
+            if self.SecondaryDataDataVal != 0:
+                Expr = '(bool)((((unsigned int*)'+str(self.SecondaryDataDataVal)+')['+str(index)+'/32] >> '+str(index)+') & 1)'
+                HasObject = 1 ##self.AllocationFlagsSecondaryDataData.CreateValueFromExpression('['+str(index)+']', Expr).GetValueAsUnsigned(0)
+            else:
+                Expr = '(bool)((((unsigned int*)'+str(self.AllocationFlagsInlineDataAddr)+')['+str(index)+'/32] >> '+str(index)+') & 1)'
+                HasObject = 1 ##self.AllocationFlagsInlineData.CreateValueFromExpression('['+str(index)+']', Expr).GetValueAsUnsigned(0)
+            
+            if HasObject == 1:
+                return self.AllocatorInstanceData.CreateChildAtOffset('['+str(index)+']',offset,self.ElementType)
+            else:
+                return self.valobj.CreateValueFromExpression('['+str(index)+']', '(void*)0xDEADBEEF')
+        except:
+            return None
+    
+    def extract_type(self):
+        logger = lldb.formatters.Logger.Logger()
+        ArrayType = self.Pairs.GetType().GetUnqualifiedType()
+        if ArrayType.IsReferenceType():
+            ArrayType = ArrayType.GetDereferencedType()
+        elif ArrayType.IsPointerType():
+            ArrayType = ArrayType.GetPointeeType()
+        if ArrayType.GetNumberOfTemplateArguments() > 0:
+            ElementType = ArrayType.GetTemplateArgumentType(0)
+        else:
+            ElementType = None
+        return ElementType
+    
+    def update(self):
+        logger = lldb.formatters.Logger.Logger()
+        try:
+            self.Pairs = self.valobj.GetChildMemberWithName('Pairs')
+            self.Elements = self.Pairs.GetChildMemberWithName('Elements')
+            self.ElementsData = self.Elements.GetChildMemberWithName('Data')
+            self.ArrayNum = self.ElementsData.GetChildMemberWithName('ArrayNum')
+            self.NumFreeIndices = self.Elements.GetChildMemberWithName('NumFreeIndices')
+            self.AllocationFlags = self.Elements.GetChildMemberWithName('AllocationFlags')
+            self.AllocationFlagsAllocatorInstance = self.AllocationFlags.GetChildMemberWithName('AllocatorInstance')
+            self.AllocatorInstance = self.ElementsData.GetChildMemberWithName('AllocatorInstance')
+            self.AllocatorInstanceData = self.AllocatorInstance.GetChildMemberWithName('Data')
+            self.AllocationFlagsInlineData = self.AllocationFlagsAllocatorInstance.GetChildMemberWithName('InlineData')
+            self.AllocationFlagsInlineDataAddr = self.AllocationFlagsInlineData.AddressOf().GetValueAsUnsigned(0)
+            self.AllocationFlagsSecondaryData = self.AllocationFlagsAllocatorInstance.GetChildMemberWithName('SecondaryData')
+            self.AllocationFlagsSecondaryDataData = self.AllocationFlagsSecondaryData.GetChildMemberWithName('Data')
+            self.SecondaryDataDataVal = self.AllocationFlagsSecondaryDataData.GetValueAsUnsigned(0)
+            self.ElementType = self.extract_type()
+            # Can't cast to TSetElement - the template instantiation won't allow it
+            Expr = '(size_t)sizeof(FSetElementId) + sizeof(int32)'
+            TSetElement = self.valobj.CreateValueFromExpression('TSetElement', Expr)
+            # This may fail due to C++ struct padding - will have to check
+            self.ElementTypeSize = self.ElementType.GetByteSize() + TSetElement.GetValueAsUnsigned(0)
+            assert self.ElementTypeSize != 0
+        except:
+            pass
+    
+    def has_children(self):
+        return True
+
+def UE4MapSummaryProvider(valobj,dict):
+    return 'size=' + str(valobj.GetNumChildren())
+
+def __lldb_init_module(debugger,dict):
+    debugger.HandleCommand('type summary add -F UE4DataFormatters.UE4TCharSummaryProvider -e TCHAR -w UE4DataFormatters')
+    debugger.HandleCommand('type summary add -F UE4DataFormatters.UE4FStringSummaryProvider -e -x "FString$" -w UE4DataFormatters')
+    debugger.HandleCommand('type summary add -F UE4DataFormatters.UE4FNameSummaryProvider -e -x "FName$" -w UE4DataFormatters')
+    debugger.HandleCommand('type summary add -F UE4DataFormatters.UE4FMinimalNameSummaryProvider -e -x "UE4FMinimalName$" -w UE4DataFormatters')
+    debugger.HandleCommand('type summary add -F UE4DataFormatters.UE4UObjectBaseSummaryProvider -e UObject -w UE4DataFormatters')
+    debugger.HandleCommand('type summary add -F UE4DataFormatters.UE4UObjectBaseSummaryProvider -e UObjectBase -w UE4DataFormatters')
+    debugger.HandleCommand('type summary add -F UE4DataFormatters.UE4UObjectBaseSummaryProvider -e UObjectBaseUtility -w UE4DataFormatters')
+    debugger.HandleCommand('type summary add -F UE4DataFormatters.UE4FWeakObjectPtrSummaryProvider -e FWeakObjectPtr -w UE4DataFormatters')
+    debugger.HandleCommand('type synthetic add -l UE4DataFormatters.UE4TWeakObjectPtrSynthProvider -x "TWeakObjectPtr<.+>$" -w UE4DataFormatters')
+    debugger.HandleCommand('type synthetic add -l UE4DataFormatters.UE4TWeakObjectPtrSynthProvider -x "TAutoWeakObjectPtr<.+>$" -w UE4DataFormatters')
+    debugger.HandleCommand('type synthetic add -l UE4DataFormatters.UE4ArraySynthProvider -x "TArray<.+,.+>$" -w UE4DataFormatters')
+    debugger.HandleCommand('type summary add -F UE4DataFormatters.UE4ArraySummaryProvider -e -x "TArray<.+>$" -w UE4DataFormatters')
+    debugger.HandleCommand('type synthetic add -l UE4DataFormatters.UE4BitArraySynthProvider -x "TBitArray<.+>$" -w UE4DataFormatters')
+    debugger.HandleCommand('type summary add -F UE4DataFormatters.UE4BitArraySummaryProvider -e -x "TBitArray<.+>$" -w UE4DataFormatters')
+    debugger.HandleCommand('type synthetic add -l UE4DataFormatters.UE4SparseArraySynthProvider -x "TSparseArray<.+>$" -w UE4DataFormatters')
+    debugger.HandleCommand('type summary add -F UE4DataFormatters.UE4SparseArraySummaryProvider -e -x "TSparseArray<.+>$" -w UE4DataFormatters')
+    debugger.HandleCommand('type synthetic add -l UE4DataFormatters.UE4ChunkedArraySynthProvider -x "TChunkedArray<.+>$" -w UE4DataFormatters')
+    debugger.HandleCommand('type summary add -F UE4DataFormatters.UE4ChunkedArraySummaryProvider -e -x "TChunkedArray<.+>$" -w UE4DataFormatters')
+    debugger.HandleCommand('type synthetic add -l UE4DataFormatters.UE4SetSynthProvider -x "TSet<.+>$" -w UE4DataFormatters')
+    debugger.HandleCommand('type summary add -F UE4DataFormatters.UE4SetSummaryProvider -e -x "TSet<.+>$" -w UE4DataFormatters')
+    debugger.HandleCommand('type synthetic add -l UE4DataFormatters.UE4MapSynthProvider -x "TMap<.+>$" -w UE4DataFormatters')
+    debugger.HandleCommand('type summary add -F UE4DataFormatters.UE4MapSummaryProvider -e -x "TMap<.+>$" -w UE4DataFormatters')
+    debugger.HandleCommand('type synthetic add -l UE4DataFormatters.UE4MapSynthProvider -x "TMapBase<.+>$" -w UE4DataFormatters')
+    debugger.HandleCommand('type summary add -F UE4DataFormatters.UE4MapSummaryProvider -e -x "TMapBase<.+>$" -w UE4DataFormatters')
+    debugger.HandleCommand("type category enable UE4DataFormatters")
+
diff --git a/Engine/Source/Runtime/Core/Private/Android/AndroidJavaEnv.cpp b/Engine/Source/Runtime/Core/Private/Android/AndroidJavaEnv.cpp
--- a/Engine/Source/Runtime/Core/Private/Android/AndroidJavaEnv.cpp
+++ b/Engine/Source/Runtime/Core/Private/Android/AndroidJavaEnv.cpp
@@ -147,7 +147,7 @@ JNIEnv* AndroidJavaEnv::GetJavaEnv( bool bRequireGlobalThis /*= true*/ )
 		jint AttachResult = CurrentJavaVM->AttachCurrentThread(&Env, NULL);
 		if (AttachResult == JNI_ERR)
 		{
-			FPlatformMisc::LowLevelOutputDebugString(L"UNIT TEST -- Failed to attach thread to get the JNI environment!");
+			FPlatformMisc::LowLevelOutputDebugString(TEXT("UNIT TEST -- Failed to attach thread to get the JNI environment!"));
 			check(false);
 			return nullptr;
 		}
@@ -155,7 +155,7 @@ JNIEnv* AndroidJavaEnv::GetJavaEnv( bool bRequireGlobalThis /*= true*/ )
 	}
 	else if (GetEnvResult != JNI_OK)
 	{
-		FPlatformMisc::LowLevelOutputDebugStringf(L"UNIT TEST -- Failed to get the JNI environment! Result = %d", GetEnvResult);
+		FPlatformMisc::LowLevelOutputDebugStringf(TEXT("UNIT TEST -- Failed to get the JNI environment! Result = %d"), GetEnvResult);
 		check(false);
 		return nullptr;
 
diff --git a/Engine/Source/Runtime/Core/Private/Android/AndroidString.cpp b/Engine/Source/Runtime/Core/Private/Android/AndroidString.cpp
--- a/Engine/Source/Runtime/Core/Private/Android/AndroidString.cpp
+++ b/Engine/Source/Runtime/Core/Private/Android/AndroidString.cpp
@@ -5,6 +5,9 @@
 =============================================================================*/
 
 #include "Android/AndroidString.h"
+
+#if !PLATFORM_TCHAR_IS_CHAR16
+
 #include "Containers/StringConv.h"
 #include <stdlib.h>
 #include <cwchar>
@@ -312,3 +315,5 @@ int vswprintf( TCHAR *buf, int max, const TCHAR *fmt, va_list args )
 	*dst = 0;  // null terminate the new string.
 	return(dst - buf);
 }
+
+#endif
\ No newline at end of file
diff --git a/Engine/Source/Runtime/Core/Public/Android/AndroidPlatform.h b/Engine/Source/Runtime/Core/Public/Android/AndroidPlatform.h
--- a/Engine/Source/Runtime/Core/Public/Android/AndroidPlatform.h
+++ b/Engine/Source/Runtime/Core/Public/Android/AndroidPlatform.h
@@ -20,7 +20,8 @@ struct FAndroidTypes : public FGenericPlatformTypes
 	//typedef unsigned int				DWORD;
 	//typedef size_t					SIZE_T;
 	//typedef decltype(NULL)			TYPE_OF_NULL;
-	typedef char16_t					CHAR16;
+	typedef char16_t					WIDECHAR;
+	typedef WIDECHAR					TCHAR;
 };
 
 typedef FAndroidTypes FPlatformTypes;
@@ -41,7 +42,7 @@ typedef FAndroidTypes FPlatformTypes;
 #define PLATFORM_MAX_FILEPATH_LENGTH				MAX_PATH
 #define PLATFORM_SUPPORTS_TEXTURE_STREAMING			1
 #define PLATFORM_REQUIRES_FILESERVER				1
-#define PLATFORM_TCHAR_IS_4_BYTES					1
+#define PLATFORM_TCHAR_IS_CHAR16					1
 #define PLATFORM_HAS_NO_EPROCLIM					1
 #define PLATFORM_USES_ES2							1
 #define PLATFORM_BUILTIN_VERTEX_HALF_FLOAT			0
diff --git a/Engine/Source/Runtime/Core/Public/Android/AndroidString.h b/Engine/Source/Runtime/Core/Public/Android/AndroidString.h
--- a/Engine/Source/Runtime/Core/Public/Android/AndroidString.h
+++ b/Engine/Source/Runtime/Core/Public/Android/AndroidString.h
@@ -5,19 +5,36 @@
 	AndroidString.h: Android platform string classes
 ==============================================================================================*/
 
-//@todo android: probably rewrite/simplify most of this.  currently it converts all wide chars
+#pragma once
+
+#include "Misc/Build.h"
+#include "HAL/Platform.h"
+
+/**
+ * Android string implementation
+ **/
+
+#if PLATFORM_TCHAR_IS_CHAR16
+
+// By default we now use 2-byte strings on Android.
+
+#include "GenericPlatform/GenericWidePlatformString.h"
+
+typedef FGenericWidePlatformString FAndroidPlatformString;
+
+#else
+
+// NOTE: This legacy 4-byte implementation will probably be removed in the future.
+
+// @todo android: probably rewrite/simplify most of this.  currently it converts all wide chars
 // to ANSI.  The Android NDK appears to have wchar_t support, but many of the functions appear
 // to just be stubs to the non-wide versions.  This doesn't work for obvious reasons.
 
-#pragma once
 #include "Misc/Char.h"
 #include "GenericPlatform/GenericPlatformMemory.h"
 #include "GenericPlatform/GenericPlatformStricmp.h"
 #include "GenericPlatform/GenericPlatformString.h"
 
-/**
- * Android string implementation
- **/
 struct FAndroidPlatformString : public FGenericPlatformString
 {
 	template <typename CharType>
@@ -581,6 +598,6 @@ struct FAndroidPlatformString : public FGenericPlatformString
 		return Result;
 	}
 };
-
+#endif
 
 typedef FAndroidPlatformString FPlatformString;
diff --git a/Engine/Source/Runtime/Core/Public/GenericPlatform/GenericWidePlatformString.h b/Engine/Source/Runtime/Core/Public/GenericPlatform/GenericWidePlatformString.h
--- a/Engine/Source/Runtime/Core/Public/GenericPlatform/GenericWidePlatformString.h
+++ b/Engine/Source/Runtime/Core/Public/GenericPlatform/GenericWidePlatformString.h
@@ -7,7 +7,7 @@
 #include "GenericPlatform/GenericPlatformStricmp.h"
 #include "GenericPlatform/GenericPlatformString.h"
 
-#if PLATFORM_APPLE || PLATFORM_LINUX || PLATFORM_HTML5 || PLATFORM_PS4 || PLATFORM_SWITCH
+#if PLATFORM_APPLE || PLATFORM_LINUX || PLATFORM_HTML5 || PLATFORM_PS4 || PLATFORM_SWITCH || PLATFORM_ANDROID
 
 /**
 * Standard implementation
diff --git a/Engine/Source/Runtime/Launch/Private/Android/LaunchAndroid.cpp b/Engine/Source/Runtime/Launch/Private/Android/LaunchAndroid.cpp
--- a/Engine/Source/Runtime/Launch/Private/Android/LaunchAndroid.cpp
+++ b/Engine/Source/Runtime/Launch/Private/Android/LaunchAndroid.cpp
@@ -285,7 +285,7 @@ extern void AndroidThunkCpp_DismissSplashScreen();
 //Main function called from the android entry point
 int32 AndroidMain(struct android_app* state)
 {
-	FPlatformMisc::LowLevelOutputDebugString(L"Entered AndroidMain()\n");
+	FPlatformMisc::LowLevelOutputDebugString(TEXT("Entered AndroidMain()\n"));
 
 	// Force the first call to GetJavaEnv() to happen on the game thread, allowing subsequent calls to occur on any thread
 	FAndroidApplication::GetJavaEnv();
@@ -360,7 +360,7 @@ int32 AndroidMain(struct android_app* state)
 	FPlatformMisc::LowLevelOutputDebugStringf(TEXT("Final commandline: %s\n"), FCommandLine::Get());
 
 	EventHandlerEvent = FPlatformProcess::GetSynchEventFromPool(false);
-	FPlatformMisc::LowLevelOutputDebugString(L"Created sync event\n");
+	FPlatformMisc::LowLevelOutputDebugString(TEXT("Created sync event\n"));
 	FAppEventManager::GetInstance()->SetEventHandlerEvent(EventHandlerEvent);
 
 	// ready for onCreate to complete
@@ -552,21 +552,21 @@ static void* AndroidEventThreadWorker( void* param )
 
 	FPlatformProcess::SetThreadAffinityMask(FPlatformAffinity::GetMainGameMask());
 
-	FPlatformMisc::LowLevelOutputDebugString(L"Entering event processing thread engine entry point");
+	FPlatformMisc::LowLevelOutputDebugString(TEXT("Entering event processing thread engine entry point"));
 
 	ALooper* looper = ALooper_prepare(ALOOPER_PREPARE_ALLOW_NON_CALLBACKS);
 	ALooper_addFd(looper, state->msgread, LOOPER_ID_MAIN, ALOOPER_EVENT_INPUT, NULL,
 		&state->cmdPollSource);
 	state->looper = looper;
 
-	FPlatformMisc::LowLevelOutputDebugString(L"Prepared looper for event thread");
+	FPlatformMisc::LowLevelOutputDebugString(TEXT("Prepared looper for event thread"));
 
 	//Assign the callbacks
 	state->onAppCmd = OnAppCommandCB;
 	state->onInputEvent = HandleInputCB;
 
-	FPlatformMisc::LowLevelOutputDebugString(L"Passed callback initialization");
-	FPlatformMisc::LowLevelOutputDebugString(L"Passed sensor initialization");
+	FPlatformMisc::LowLevelOutputDebugString(TEXT("Passed callback initialization"));
+	FPlatformMisc::LowLevelOutputDebugString(TEXT("Passed sensor initialization"));
 
 #if WITH_Android_Choreographer
 	TheChoreographer.SetupChoreographer();
@@ -575,7 +575,7 @@ static void* AndroidEventThreadWorker( void* param )
 	//continue to process events until the engine is shutting down
 	while (!GIsRequestingExit)
 	{
-//		FPlatformMisc::LowLevelOutputDebugString(L"AndroidEventThreadWorker");
+//		FPlatformMisc::LowLevelOutputDebugString(TEXT("AndroidEventThreadWorker"));
 
 		AndroidProcessEvents(state);
 
@@ -611,7 +611,7 @@ struct android_app* GNativeAndroidApp = NULL;
 
 void android_main(struct android_app* state)
 {
-	FPlatformMisc::LowLevelOutputDebugString(L"Entering native app glue main function");
+	FPlatformMisc::LowLevelOutputDebugString(TEXT("Entering native app glue main function"));
 	
 	GNativeAndroidApp = state;
 	check(GNativeAndroidApp);
@@ -621,7 +621,7 @@ void android_main(struct android_app* state)
 	pthread_attr_setdetachstate(&otherAttr, PTHREAD_CREATE_DETACHED);
 	pthread_create(&G_AndroidEventThread, &otherAttr, AndroidEventThreadWorker, state);
 
-	FPlatformMisc::LowLevelOutputDebugString(L"Created event thread");
+	FPlatformMisc::LowLevelOutputDebugString(TEXT("Created event thread"));
 
 	// Make sure glue isn't stripped. (not needed in ndk-15)
 #if PLATFORM_ANDROID_NDK_VERSION < 150000
@@ -637,7 +637,7 @@ extern bool GAndroidGPUInfoReady;
 //Called from the event process thread
 static int32_t HandleInputCB(struct android_app* app, AInputEvent* event)
 {
-//	FPlatformMisc::LowLevelOutputDebugStringf(L"INPUT - type: %x, action: %x, source: %x, keycode: %x, buttons: %x", AInputEvent_getType(event), 
+//	FPlatformMisc::LowLevelOutputDebugStringf(TEXT("INPUT - type: %x, action: %x, source: %x, keycode: %x, buttons: %x"), AInputEvent_getType(event), 
 //		AMotionEvent_getAction(event), AInputEvent_getSource(event), AKeyEvent_getKeyCode(event), AMotionEvent_getButtonState(event));
 
 	int32 EventType = AInputEvent_getType(event);
@@ -891,7 +891,7 @@ static int32_t HandleInputCB(struct android_app* app, AInputEvent* event)
 		int keyFlags = AKeyEvent_getFlags(event);
 		bool bSoftKey = (keyFlags & AKEY_EVENT_FLAG_SOFT_KEYBOARD) != 0;
 
-		FPlatformMisc::LowLevelOutputDebugStringf(L"Received keycode: %d, softkey: %d", keyCode, bSoftKey ? 1 : 0);
+		FPlatformMisc::LowLevelOutputDebugStringf(TEXT("Received keycode: %d, softkey: %d"), keyCode, bSoftKey ? 1 : 0);
 
 		//Trap codes handled as possible gamepad events
 		if (ValidGamepadKeyCodes.Contains(keyCode))
@@ -906,11 +906,11 @@ static int32_t HandleInputCB(struct android_app* app, AInputEvent* event)
 			}
 			bool down = AKeyEvent_getAction(event) != AKEY_EVENT_ACTION_UP;
 			FAndroidInputInterface::JoystickButtonEvent(device, keyCode, down);
-			FPlatformMisc::LowLevelOutputDebugStringf(L"Received gamepad button: %d", keyCode);
+			FPlatformMisc::LowLevelOutputDebugStringf(TEXT("Received gamepad button: %d"), keyCode);
 		}
 		else
 		{
-			FPlatformMisc::LowLevelOutputDebugStringf(L"Received key event: %d", keyCode);
+			FPlatformMisc::LowLevelOutputDebugStringf(TEXT("Received key event: %d"), keyCode);
 
 			// only handle mapped key codes
 			if (!MappedKeyCodes.Contains(keyCode))
@@ -956,7 +956,7 @@ static bool IsStartupMoviePlaying()
 static void OnAppCommandCB(struct android_app* app, int32_t cmd)
 {
 	bool bNeedToSync = false;
-	//FPlatformMisc::LowLevelOutputDebugStringf(L"OnAppCommandCB cmd: %u, tid = %d", cmd, gettid());
+	//FPlatformMisc::LowLevelOutputDebugStringf(TEXT("OnAppCommandCB cmd: %u, tid = %d"), cmd, gettid());
 
 	switch (cmd)
 	{
@@ -1138,7 +1138,7 @@ static void OnAppCommandCB(struct android_app* app, int32_t cmd)
 	{
 		FAppEventManager::GetInstance()->WaitForEmptyQueue();
 	}
-	//FPlatformMisc::LowLevelOutputDebugStringf(L"#### END OF OnAppCommandCB cmd: %u, tid = %d", cmd, gettid());
+	//FPlatformMisc::LowLevelOutputDebugStringf(TEXT("#### END OF OnAppCommandCB cmd: %u, tid = %d"), cmd, gettid());
 }
 
 //Native-defined functions
