Fix issues with LLDB data formatters with 2 byte characters

#ue4

diff --git a/Engine/Extras/LLDBDataFormatters/UE4DataFormatters_2ByteChars.py b/Engine/Extras/LLDBDataFormatters/UE4DataFormatters_2ByteChars.py
--- a/Engine/Extras/LLDBDataFormatters/UE4DataFormatters_2ByteChars.py
+++ b/Engine/Extras/LLDBDataFormatters/UE4DataFormatters_2ByteChars.py
@@ -30,11 +30,11 @@ def UE4TCharSummaryProvider(valobj,dict):
         if DataVal == 0:
             Val = 'NULL'
         else:
-            Expr = '(ushort_t*)('+str(Data)+')'
+            Expr = '(char16_t*)('+str(Data)+')'
             ValRef = valobj.CreateValueFromExpression('string', Expr)
             Val = str(ValRef.GetSummary())
     elif Type.IsReferenceType():
-        Expr = '(ushort_t&)('+str(Data)+')'
+        Expr = '(char16_t&)('+str(Data)+')'
         ValRef = valobj.CreateValueFromExpression('string', Expr)
         Val = str(ValRef.GetSummary())
     return Val
@@ -49,7 +49,7 @@ def UE4FStringSummaryProvider(valobj,dict):
     else:
         AllocatorInstance = Data.GetChildMemberWithName('AllocatorInstance')
         ActualData = AllocatorInstance.GetChildMemberWithName('Data')
-        Expr = '(ushort_t*)('+str(ActualData.GetValue())+')'
+        Expr = '(char16_t*)('+str(ActualData.GetValue())+')'
         ValRef = valobj.CreateValueFromExpression('string', Expr)
         Val = str(ValRef.GetSummary())
         return 'string=' + Val
@@ -638,29 +638,29 @@ def UE4MapSummaryProvider(valobj,dict):
     return 'size=' + str(valobj.GetNumChildren())
 
 def __lldb_init_module(debugger,dict):
-    debugger.HandleCommand('type summary add -F UE4DataFormatters.UE4TCharSummaryProvider -e TCHAR -w UE4DataFormatters')
-    debugger.HandleCommand('type summary add -F UE4DataFormatters.UE4FStringSummaryProvider -e -x "FString$" -w UE4DataFormatters')
-    debugger.HandleCommand('type summary add -F UE4DataFormatters.UE4FNameSummaryProvider -e -x "FName$" -w UE4DataFormatters')
-    debugger.HandleCommand('type summary add -F UE4DataFormatters.UE4FMinimalNameSummaryProvider -e -x "UE4FMinimalName$" -w UE4DataFormatters')
-    debugger.HandleCommand('type summary add -F UE4DataFormatters.UE4UObjectBaseSummaryProvider -e UObject -w UE4DataFormatters')
-    debugger.HandleCommand('type summary add -F UE4DataFormatters.UE4UObjectBaseSummaryProvider -e UObjectBase -w UE4DataFormatters')
-    debugger.HandleCommand('type summary add -F UE4DataFormatters.UE4UObjectBaseSummaryProvider -e UObjectBaseUtility -w UE4DataFormatters')
-    debugger.HandleCommand('type summary add -F UE4DataFormatters.UE4FWeakObjectPtrSummaryProvider -e FWeakObjectPtr -w UE4DataFormatters')
-    debugger.HandleCommand('type synthetic add -l UE4DataFormatters.UE4TWeakObjectPtrSynthProvider -x "TWeakObjectPtr<.+>$" -w UE4DataFormatters')
-    debugger.HandleCommand('type synthetic add -l UE4DataFormatters.UE4TWeakObjectPtrSynthProvider -x "TAutoWeakObjectPtr<.+>$" -w UE4DataFormatters')
-    debugger.HandleCommand('type synthetic add -l UE4DataFormatters.UE4ArraySynthProvider -x "TArray<.+,.+>$" -w UE4DataFormatters')
-    debugger.HandleCommand('type summary add -F UE4DataFormatters.UE4ArraySummaryProvider -e -x "TArray<.+>$" -w UE4DataFormatters')
-    debugger.HandleCommand('type synthetic add -l UE4DataFormatters.UE4BitArraySynthProvider -x "TBitArray<.+>$" -w UE4DataFormatters')
-    debugger.HandleCommand('type summary add -F UE4DataFormatters.UE4BitArraySummaryProvider -e -x "TBitArray<.+>$" -w UE4DataFormatters')
-    debugger.HandleCommand('type synthetic add -l UE4DataFormatters.UE4SparseArraySynthProvider -x "TSparseArray<.+>$" -w UE4DataFormatters')
-    debugger.HandleCommand('type summary add -F UE4DataFormatters.UE4SparseArraySummaryProvider -e -x "TSparseArray<.+>$" -w UE4DataFormatters')
-    debugger.HandleCommand('type synthetic add -l UE4DataFormatters.UE4ChunkedArraySynthProvider -x "TChunkedArray<.+>$" -w UE4DataFormatters')
-    debugger.HandleCommand('type summary add -F UE4DataFormatters.UE4ChunkedArraySummaryProvider -e -x "TChunkedArray<.+>$" -w UE4DataFormatters')
-    debugger.HandleCommand('type synthetic add -l UE4DataFormatters.UE4SetSynthProvider -x "TSet<.+>$" -w UE4DataFormatters')
-    debugger.HandleCommand('type summary add -F UE4DataFormatters.UE4SetSummaryProvider -e -x "TSet<.+>$" -w UE4DataFormatters')
-    debugger.HandleCommand('type synthetic add -l UE4DataFormatters.UE4MapSynthProvider -x "TMap<.+>$" -w UE4DataFormatters')
-    debugger.HandleCommand('type summary add -F UE4DataFormatters.UE4MapSummaryProvider -e -x "TMap<.+>$" -w UE4DataFormatters')
-    debugger.HandleCommand('type synthetic add -l UE4DataFormatters.UE4MapSynthProvider -x "TMapBase<.+>$" -w UE4DataFormatters')
-    debugger.HandleCommand('type summary add -F UE4DataFormatters.UE4MapSummaryProvider -e -x "TMapBase<.+>$" -w UE4DataFormatters')
+    debugger.HandleCommand('type summary add -F UE4DataFormatters_2ByteChars.UE4TCharSummaryProvider -e TCHAR -w UE4DataFormatters')
+    debugger.HandleCommand('type summary add -F UE4DataFormatters_2ByteChars.UE4FStringSummaryProvider -e -x "FString$" -w UE4DataFormatters')
+    debugger.HandleCommand('type summary add -F UE4DataFormatters_2ByteChars.UE4FNameSummaryProvider -e -x "FName$" -w UE4DataFormatters')
+    debugger.HandleCommand('type summary add -F UE4DataFormatters_2ByteChars.UE4FMinimalNameSummaryProvider -e -x "UE4FMinimalName$" -w UE4DataFormatters')
+    debugger.HandleCommand('type summary add -F UE4DataFormatters_2ByteChars.UE4UObjectBaseSummaryProvider -e UObject -w UE4DataFormatters')
+    debugger.HandleCommand('type summary add -F UE4DataFormatters_2ByteChars.UE4UObjectBaseSummaryProvider -e UObjectBase -w UE4DataFormatters')
+    debugger.HandleCommand('type summary add -F UE4DataFormatters_2ByteChars.UE4UObjectBaseSummaryProvider -e UObjectBaseUtility -w UE4DataFormatters')
+    debugger.HandleCommand('type summary add -F UE4DataFormatters_2ByteChars.UE4FWeakObjectPtrSummaryProvider -e FWeakObjectPtr -w UE4DataFormatters')
+    debugger.HandleCommand('type synthetic add -l UE4DataFormatters_2ByteChars.UE4TWeakObjectPtrSynthProvider -x "TWeakObjectPtr<.+>$" -w UE4DataFormatters')
+    debugger.HandleCommand('type synthetic add -l UE4DataFormatters_2ByteChars.UE4TWeakObjectPtrSynthProvider -x "TAutoWeakObjectPtr<.+>$" -w UE4DataFormatters')
+    debugger.HandleCommand('type synthetic add -l UE4DataFormatters_2ByteChars.UE4ArraySynthProvider -x "TArray<.+,.+>$" -w UE4DataFormatters')
+    debugger.HandleCommand('type summary add -F UE4DataFormatters_2ByteChars.UE4ArraySummaryProvider -e -x "TArray<.+>$" -w UE4DataFormatters')
+    debugger.HandleCommand('type synthetic add -l UE4DataFormatters_2ByteChars.UE4BitArraySynthProvider -x "TBitArray<.+>$" -w UE4DataFormatters')
+    debugger.HandleCommand('type summary add -F UE4DataFormatters_2ByteChars.UE4BitArraySummaryProvider -e -x "TBitArray<.+>$" -w UE4DataFormatters')
+    debugger.HandleCommand('type synthetic add -l UE4DataFormatters_2ByteChars.UE4SparseArraySynthProvider -x "TSparseArray<.+>$" -w UE4DataFormatters')
+    debugger.HandleCommand('type summary add -F UE4DataFormatters_2ByteChars.UE4SparseArraySummaryProvider -e -x "TSparseArray<.+>$" -w UE4DataFormatters')
+    debugger.HandleCommand('type synthetic add -l UE4DataFormatters_2ByteChars.UE4ChunkedArraySynthProvider -x "TChunkedArray<.+>$" -w UE4DataFormatters')
+    debugger.HandleCommand('type summary add -F UE4DataFormatters_2ByteChars.UE4ChunkedArraySummaryProvider -e -x "TChunkedArray<.+>$" -w UE4DataFormatters')
+    debugger.HandleCommand('type synthetic add -l UE4DataFormatters_2ByteChars.UE4SetSynthProvider -x "TSet<.+>$" -w UE4DataFormatters')
+    debugger.HandleCommand('type summary add -F UE4DataFormatters_2ByteChars.UE4SetSummaryProvider -e -x "TSet<.+>$" -w UE4DataFormatters')
+    debugger.HandleCommand('type synthetic add -l UE4DataFormatters_2ByteChars.UE4MapSynthProvider -x "TMap<.+>$" -w UE4DataFormatters')
+    debugger.HandleCommand('type summary add -F UE4DataFormatters_2ByteChars.UE4MapSummaryProvider -e -x "TMap<.+>$" -w UE4DataFormatters')
+    debugger.HandleCommand('type synthetic add -l UE4DataFormatters_2ByteChars.UE4MapSynthProvider -x "TMapBase<.+>$" -w UE4DataFormatters')
+    debugger.HandleCommand('type summary add -F UE4DataFormatters_2ByteChars.UE4MapSummaryProvider -e -x "TMapBase<.+>$" -w UE4DataFormatters')
     debugger.HandleCommand("type category enable UE4DataFormatters")
 
