Fix to debugger visualization of FWeakObjectPtr and TWeakObjectPtr.

https://udn.unrealengine.com/questions/455158/tweakobjectptr.html

diff --git a/Engine/Extras/LLDBDataFormatters/UE4DataFormatters.py b/Engine/Extras/LLDBDataFormatters/UE4DataFormatters.py
--- a/Engine/Extras/LLDBDataFormatters/UE4DataFormatters.py
+++ b/Engine/Extras/LLDBDataFormatters/UE4DataFormatters.py
@@ -108,11 +108,11 @@ class UE4TWeakObjectPtrSynthProvider:
         logger = lldb.formatters.Logger.Logger()
         logger >> "Retrieving child " + str(index)
         if self.ObjectSerialNumberVal >= 1:
-            Expr = 'GObjectArrayForDebugVisualizers->Objects['+str(self.ObjectIndexVal)+'].SerialNumber == '+str(self.ObjectSerialNumberVal)
+            Expr = 'GObjectArrayForDebugVisualizers->Objects['+str(self.ObjectIndexVal/65536)+']['+str(self.ObjectIndexVal%65536)+'].SerialNumber == '+str(self.ObjectSerialNumberVal)
             Val = self.valobj.CreateValueFromExpression(str(self.ObjectIndexVal), Expr)
             Value = Val.GetValueAsUnsigned(0)
             if Value != 0:
-                Expr = 'GObjectArrayForDebugVisualizers->Objects['+str(self.ObjectIndexVal)+'].Object'
+                Expr = 'GObjectArrayForDebugVisualizers->Objects['+str(self.ObjectIndexVal/65536)+']['+str(self.ObjectIndexVal%65536)+'].Object'
                 return self.valobj.CreateValueFromExpression('Object', Expr)
             else:
                 Expr = '(void*)0xDEADBEEF'
@@ -141,13 +141,13 @@ def UE4FWeakObjectPtrSummaryProvider(valobj,dict):
         return 'object=nullptr'
     ObjectIndex = valobj.GetChildMemberWithName('ObjectIndex')
     ObjectIndexVal = ObjectIndex.GetValueAsSigned(0)
-    Expr = 'GObjectArrayForDebugVisualizers->Objects['+str(ObjectIndexVal)+'].SerialNumber == '+str(ObjectSerialNumberVal)
+    Expr = 'GObjectArrayForDebugVisualizers->Objects['+str(ObjectIndexVal/65536)+']['+str(ObjectIndexVal%65536)+'].SerialNumber == '+str(ObjectSerialNumberVal)
     Val = valobj.CreateValueFromExpression(str(ObjectIndexVal), Expr)
     ValRef = Val.GetValueAsUnsigned(0)
     if ValRef == 0:
         return 'object=STALE'
     else:
-        Expr = 'GObjectArrayForDebugVisualizers->Objects['+str(ObjectIndexVal)+'].Object'
+        Expr = 'GObjectArrayForDebugVisualizers->Objects['+str(ObjectIndexVal/65536)+']['+str(ObjectIndexVal%65536)+'].Object'
         Val = valobj.CreateValueFromExpression(str(ObjectIndexVal), Expr)
         return 'object=' + Val.GetValue()
 
diff --git a/Engine/Extras/LLDBDataFormatters/UE4DataFormatters_2ByteChars.py b/Engine/Extras/LLDBDataFormatters/UE4DataFormatters_2ByteChars.py
--- a/Engine/Extras/LLDBDataFormatters/UE4DataFormatters_2ByteChars.py
+++ b/Engine/Extras/LLDBDataFormatters/UE4DataFormatters_2ByteChars.py
@@ -108,11 +108,11 @@ class UE4TWeakObjectPtrSynthProvider:
         logger = lldb.formatters.Logger.Logger()
         logger >> "Retrieving child " + str(index)
         if self.ObjectSerialNumberVal >= 1:
-            Expr = 'GObjectArrayForDebugVisualizers->Objects['+str(self.ObjectIndexVal)+'].SerialNumber == '+str(self.ObjectSerialNumberVal)
+            Expr = 'GObjectArrayForDebugVisualizers->Objects['+str(self.ObjectIndexVal/65536)+']['+str(self.ObjectIndexVal%65536)+'].SerialNumber == '+str(self.ObjectSerialNumberVal)
             Val = self.valobj.CreateValueFromExpression(str(self.ObjectIndexVal), Expr)
             Value = Val.GetValueAsUnsigned(0)
             if Value != 0:
-                Expr = 'GObjectArrayForDebugVisualizers->Objects['+str(self.ObjectIndexVal)+'].Object'
+                Expr = 'GObjectArrayForDebugVisualizers->Objects['+str(self.ObjectIndexVal/65536)+']['+str(self.ObjectIndexVal%65536)+'].Object'
                 return self.valobj.CreateValueFromExpression('Object', Expr)
             else:
                 Expr = '(void*)0xDEADBEEF'
@@ -141,13 +141,13 @@ def UE4FWeakObjectPtrSummaryProvider(valobj,dict):
         return 'object=nullptr'
     ObjectIndex = valobj.GetChildMemberWithName('ObjectIndex')
     ObjectIndexVal = ObjectIndex.GetValueAsSigned(0)
-    Expr = 'GObjectArrayForDebugVisualizers->Objects['+str(ObjectIndexVal)+'].SerialNumber == '+str(ObjectSerialNumberVal)
+    Expr = 'GObjectArrayForDebugVisualizers->Objects['+str(ObjectIndexVal/65536)+']Objects['+str(ObjectIndexVal%65536)+'].SerialNumber == '+str(ObjectSerialNumberVal)
     Val = valobj.CreateValueFromExpression(str(ObjectIndexVal), Expr)
     ValRef = Val.GetValueAsUnsigned(0)
     if ValRef == 0:
         return 'object=STALE'
     else:
-        Expr = 'GObjectArrayForDebugVisualizers->Objects['+str(ObjectIndexVal)+'].Object'
+        Expr = 'GObjectArrayForDebugVisualizers->Objects['+str(ObjectIndexVal/65536)+']['+str(ObjectIndexVal%65536)+'].Object'
         Val = valobj.CreateValueFromExpression(str(ObjectIndexVal), Expr)
         return 'object=' + Val.GetValue()
 
