Fix debug formatter for TArrays with inline-data allocators.

diff --git a/Engine/Extras/LLDBDataFormatters/UE4DataFormatters.py b/Engine/Extras/LLDBDataFormatters/UE4DataFormatters.py
--- a/Engine/Extras/LLDBDataFormatters/UE4DataFormatters.py
+++ b/Engine/Extras/LLDBDataFormatters/UE4DataFormatters.py
@@ -386,21 +386,27 @@ class UE4ArraySynthProvider:
             return None;
         
         if index < self.NumChildren:
+            logger >> "Existing child " + str(index)
             return self.valobj.GetChildAtIndex(index)
         else:
             index -= self.NumChildren
 
         if index >= self.num_children():
+            logger >> "Index is greater than num children " + str(self.num_children())
             return None;
         try:
             offset = index * self.ElementTypeSize
-            if self.Data != None:
+            if self.Data != None and self.Data.IsValid():
+                logger >> "Retrieving Data child of type " + str(self.ElementType) + " at offset " + str(offset)
                 return self.Data.CreateChildAtOffset('['+str(index)+']',offset,self.ElementType)
             elif self.SecondaryDataDataVal > 0:
+                logger >> "Retrieving SecondaryData child of type " + str(self.ElementType) + " at offset " + str(offset)
                 return self.SecondaryDataData.CreateChildAtOffset('['+str(index)+']',offset,self.ElementType)
             else:
+                logger >> "Retrieving InlineData child of type " + str(self.ElementType) + " at offset " + str(offset)
                 return self.InlineData.CreateChildAtOffset('['+str(index)+']',offset,self.ElementType)
         except:
+            logger >> "Exception retrieving child of type " + str(self.ElementType) + " at index " + str(index)
             return None
 
     def extract_type(self):
@@ -426,6 +432,7 @@ class UE4ArraySynthProvider:
             self.AllocatorInstance = self.valobj.GetChildMemberWithName('AllocatorInstance')
             if self.AllocatorInstance.GetType().IsReferenceType():
                 self.AllocatorInstance = self.AllocatorInstance.Dereference()
+            self.Data = None
             self.Data = self.AllocatorInstance.GetChildMemberWithName('Data')
             self.InlineData = self.AllocatorInstance.GetChildMemberWithName('InlineData')
             self.SecondaryData = self.AllocatorInstance.GetChildMemberWithName('SecondaryData')
