FName debug visualization fixes

* Fixed static configs that used the wrong stride
* Also exposed some constants in NameTypes.h to improve robustness and reduce duplication of constants across UE4.natvis, PS4UE4.natvis and LLDB UE4DataFormatters.py.
* Reverted back to using DisplayIndex instead of ComparisonIndex for DisplayString
* Reverted back to using quotes around the string part, e.g. '"MyName"_123' instead of 'MyName_123'
* Speculative LLDB data formatter fix
* Fixed FStatMessage natvis special case where Cycles contained and extra FMinimalName
* Updated Natvis README.txt

diff --git a/Engine/Extras/LLDBDataFormatters/UE4DataFormatters.py b/Engine/Extras/LLDBDataFormatters/UE4DataFormatters.py
--- a/Engine/Extras/LLDBDataFormatters/UE4DataFormatters.py
+++ b/Engine/Extras/LLDBDataFormatters/UE4DataFormatters.py
@@ -62,7 +62,7 @@ def UE4FNameSummaryProvider(valobj,dict):
     if IndexVal >= 4194304:
         return 'name=Invalid'
     else:
-        Expr = '(char*)(((FNameEntry***)GFNameTableForDebuggerVisualizers_MT)['+str(IndexVal)+' / 16384]['+str(IndexVal)+' % 16384])->AnsiName'
+        Expr = '(char*)((FNameEntry&)GNameBlocksDebug['+str(IndexVal)+' >> FNameDebugVisualizer::OffsetBits][FNameDebugVisualizer::EntryStride * ('+str(IndexVal)+' & FNameDebugVisualizer::OffsetMask)]).AnsiName'
         FNameRef = valobj.CreateValueFromExpression(str(IndexVal), Expr)
         assert FNameRef != None
         Val = FNameRef.GetSummary()
@@ -77,7 +77,7 @@ def UE4FMinimalNameSummaryProvider(valobj,dict):
     if IndexVal >= 4194304:
         return 'name=Invalid'
     else:
-        Expr = '(char*)(((FNameEntry***)GFNameTableForDebuggerVisualizers_MT)['+str(IndexVal)+' / 16384]['+str(IndexVal)+' % 16384])->AnsiName'
+        Expr = '(char*)((FNameEntry&)GNameBlocksDebug['+str(IndexVal)+' >> FNameDebugVisualizer::OffsetBits][FNameDebugVisualizer::EntryStride * ('+str(IndexVal)+' & FNameDebugVisualizer::OffsetMask)]).AnsiName'
         FNameRef = valobj.CreateValueFromExpression(str(IndexVal), Expr)
         assert FNameRef != None
         Val = FNameRef.GetSummary()
diff --git a/Engine/Extras/NatvisHelpers/README.txt b/Engine/Extras/NatvisHelpers/README.txt
--- a/Engine/Extras/NatvisHelpers/README.txt
+++ b/Engine/Extras/NatvisHelpers/README.txt
@@ -2,6 +2,5 @@ This library exists to support UE's custom native visualizers in Visual Studio w
 
 When generating a patch executable, we include this library and force the linker to include a reference to the global InitNatvisHelpers() function. 
 
-This creates references to GFNameTableForDebuggerVisualizers_MT and GObjectArrayForDebugVisualizers. Object files in static libraries are only
-included if a symbol in them is referenced, so if those symbols already exist in the patch, those references will be used. Otherwise the definitions
+This creates references to GNameBlocksDebug and GObjectArrayForDebugVisualizers. Object files in static libraries are only included if a symbol in them is referenced, so if those symbols already exist in the patch, those references will be used. Otherwise the definitions
 in Global.cpp will be used.
diff --git a/Engine/Source/Runtime/Core/Private/Misc/CoreGlobals.cpp b/Engine/Source/Runtime/Core/Private/Misc/CoreGlobals.cpp
--- a/Engine/Source/Runtime/Core/Private/Misc/CoreGlobals.cpp
+++ b/Engine/Source/Runtime/Core/Private/Misc/CoreGlobals.cpp
@@ -181,7 +181,7 @@ bool					GExitPurge						= false;
 
 FChunkedFixedUObjectArray* GCoreObjectArrayForDebugVisualizers = nullptr;
 #if PLATFORM_UNIX
-uint8** CORE_API GNameBlocksDebug = FName::GetBlocksForDebugVisualizer();
+uint8** CORE_API GNameBlocksDebug = FNameDebugVisualizer::GetBlocks();
 FChunkedFixedUObjectArray*& CORE_API GObjectArrayForDebugVisualizers = GCoreObjectArrayForDebugVisualizers;
 #endif
 
diff --git a/Engine/Source/Runtime/Core/Private/UObject/UnrealNames.cpp b/Engine/Source/Runtime/Core/Private/UObject/UnrealNames.cpp
--- a/Engine/Source/Runtime/Core/Private/UObject/UnrealNames.cpp
+++ b/Engine/Source/Runtime/Core/Private/UObject/UnrealNames.cpp
@@ -1185,11 +1185,6 @@ TArray<const FNameEntry*> FName::DebugDump()
 	return GetNamePool().DebugDump();
 }
 
-uint8** FName::GetBlocksForDebugVisualizer()
-{
-	return ((FNamePool*)(NamePoolData))->GetBlocksForDebugVisualizer();
-}
-
 FNameEntry const* FName::GetEntry(EName Ename)
 {
 	FNamePool& Pool = GetNamePool();
@@ -2239,3 +2234,12 @@ void CallNameCreationHook()
 }
 
 #endif
+
+uint8** FNameDebugVisualizer::GetBlocks()
+{
+	static_assert(EntryStride == FNameEntryAllocator::Stride,	"Natvis constants out of sync with actual constants");
+	static_assert(BlockBits == FNameMaxBlockBits,				"Natvis constants out of sync with actual constants");
+	static_assert(OffsetBits == FNameBlockOffsetBits,			"Natvis constants out of sync with actual constants");
+
+	return ((FNamePool*)(NamePoolData))->GetBlocksForDebugVisualizer();
+}
diff --git a/Engine/Source/Runtime/Core/Public/Modules/Boilerplate/ModuleBoilerplate.h b/Engine/Source/Runtime/Core/Public/Modules/Boilerplate/ModuleBoilerplate.h
--- a/Engine/Source/Runtime/Core/Public/Modules/Boilerplate/ModuleBoilerplate.h
+++ b/Engine/Source/Runtime/Core/Public/Modules/Boilerplate/ModuleBoilerplate.h
@@ -49,7 +49,7 @@ class FChunkedFixedUObjectArray;
 #define UE4_VISUALIZERS_HELPERS
 #else
 #define UE4_VISUALIZERS_HELPERS \
-	uint8** GNameBlocksDebug = FName::GetBlocksForDebugVisualizer(); \
+	uint8** GNameBlocksDebug = FNameDebugVisualizer::GetBlocks(); \
 	FChunkedFixedUObjectArray*& GObjectArrayForDebugVisualizers = GCoreObjectArrayForDebugVisualizers;
 #endif
 
diff --git a/Engine/Source/Runtime/Core/Public/UObject/NameTypes.h b/Engine/Source/Runtime/Core/Public/UObject/NameTypes.h
--- a/Engine/Source/Runtime/Core/Public/UObject/NameTypes.h
+++ b/Engine/Source/Runtime/Core/Public/UObject/NameTypes.h
@@ -760,8 +760,6 @@ public:
 	static FNameEntry const* GetEntry(FNameEntryId Id);
 
 	//@}
-
-	static uint8** GetBlocksForDebugVisualizer();
 	
 	/** Run autotest on FNames. */
 	static void AutoTest();
@@ -913,3 +911,15 @@ inline void FNameEntry::Encode(WIDECHAR*, uint32) {}
 inline void FNameEntry::Decode(ANSICHAR*, uint32) {}
 inline void FNameEntry::Decode(WIDECHAR*, uint32) {}
 #endif
+
+struct FNameDebugVisualizer
+{
+	CORE_API static uint8** GetBlocks();
+private:
+	static constexpr uint32 EntryStride = alignof(FNameEntry);
+	static constexpr uint32 OffsetBits = 16;
+	static constexpr uint32 BlockBits = 13;
+	static constexpr uint32 OffsetMask = (1 << OffsetBits) - 1;
+	static constexpr uint32 UnusedMask = UINT32_MAX << BlockBits << OffsetBits;
+	static constexpr uint32 MaxLength = NAME_SIZE;
+};
\ No newline at end of file
