Miscellaneous LLDB Formater optimizations for Android Studio:
- Stop calling str().  For whatever reason, this seems to fix painfully slow debugging for certain callstacks in Android Studio.
- Fix UE4FStringSummaryProvider not handling empty strings correctly
- Map/Set synth providers now calculate size in __init__ instead of update().  This seems to provide a good speed up

I didn't do a proper optimization pass on this.  Just tried some things out and kept what felt better. This at least makes certain pathogical cases bearable (10s of seconds instead of multi-minute).  I've been running with this for a while and haven't encountered any issues.

diff --git a/Engine/Extras/LLDBDataFormatters/UE4DataFormatters_2ByteChars.py b/Engine/Extras/LLDBDataFormatters/UE4DataFormatters_2ByteChars.py
--- a/Engine/Extras/LLDBDataFormatters/UE4DataFormatters_2ByteChars.py
+++ b/Engine/Extras/LLDBDataFormatters/UE4DataFormatters_2ByteChars.py
@@ -30,18 +30,18 @@ def UE4TCharSummaryProvider(valobj,dict):
         if DataVal == 0:
             Val = 'NULL'
         else:
-            Expr = '(char16_t*)('+str(Data)+')'
+            Expr = '(char16_t*)(%s)' % Data
             ValRef = valobj.CreateValueFromExpression('string', Expr)
-            Val = str(ValRef.GetSummary())
+            Val = ValRef.GetSummary()
     elif Type.IsReferenceType():
-        Expr = '(char16_t&)('+str(Data)+')'
+        Expr = '(char16_t&)(%s)' % Data
         ValRef = valobj.CreateValueFromExpression('string', Expr)
-        Val = str(ValRef.GetSummary())
+        Val = ValRef.GetSummary()
     return Val
 
 def UE4FStringSummaryProvider(valobj,dict):
     Data = valobj.GetChildMemberWithName('Data')
-    ArrayNumVal = Data.GetNumChildren()
+    ArrayNumVal = Data.GetChildMemberWithName('ArrayNum').GetValueAsSigned(0)
     if ArrayNumVal < 0:
        return 'string=Invalid'
     elif ArrayNumVal == 0:
@@ -49,42 +49,46 @@ def UE4FStringSummaryProvider(valobj,dict):
     else:
         AllocatorInstance = Data.GetChildMemberWithName('AllocatorInstance')
         ActualData = AllocatorInstance.GetChildMemberWithName('Data')
-        Expr = '(char16_t*)('+str(ActualData.GetValue())+')'
+        Expr = '(char16_t*)(%s)' % ActualData.GetValue()
         ValRef = valobj.CreateValueFromExpression('string', Expr)
-        Val = str(ValRef.GetSummary())
+        Val = ValRef.GetSummary()
         return 'string=' + Val
 
 def UE4FNameSummaryProvider(valobj,dict):
     Index = valobj.GetChildMemberWithName('DisplayIndex')
     if not Index.IsValid():
         Index = valobj.GetChildMemberWithName('ComparisonIndex')
-    IndexVal = Index.GetValueAsUnsigned(0)
+    Number = valobj.GetChildMemberWithName('Number')
+    IndexVal = Index.GetValueAsSigned(0)
+    NumberVal = Number.GetValueAsUnsigned(0)
     if IndexVal >= 4194304:
         return 'name=Invalid'
     else:
-        Expr = '(char*)(((FNameEntry***)FName::GetNameTableForDebuggerVisualizers_MT())['+str(IndexVal)+' / 16384]['+str(IndexVal)+' % 16384])->AnsiName'
-        FNameRef = valobj.CreateValueFromExpression(str(IndexVal), Expr)
-        assert FNameRef != None
+        Expr = '(char*)(FName::GetNameTableForDebuggerVisualizers_MT()[%s / 16384][%s %% 16384]->AnsiName)' % (IndexVal, IndexVal)
+        FNameRef = valobj.CreateValueFromExpression('%s' % IndexVal, Expr)
         Val = FNameRef.GetSummary()
-        if IndexVal == 0:
-            return 'name=' + Val + '_' + str(IndexVal-1)
+        assert Val != None
+        if NumberVal != 0:
+            return 'name=%s_%s' % (Val, NumberVal-1)
         else:
-            return 'name=' + Val
+            return 'name=%s' % Val
 			
 def UE4FMinimalNameSummaryProvider(valobj,dict):
     Index = valobj.GetChildMemberWithName('Index')
-    IndexVal = Index.GetValueAsUnsigned(0)
+    Number = valobj.GetChildMemberWithName('Number')
+    IndexVal = Index.GetValueAsSigned(0)
+    NumberVal = Number.GetValueAsSigned(0)
     if IndexVal >= 4194304:
         return 'name=Invalid'
     else:
-        Expr = '(char*)(((FNameEntry***)FName::GetNameTableForDebuggerVisualizers_MT())['+str(IndexVal)+' / 16384]['+str(IndexVal)+' % 16384])->AnsiName'
-        FNameRef = valobj.CreateValueFromExpression(str(IndexVal), Expr)
+        Expr = '(char*)(((FNameEntry***)FName::GetNameTableForDebuggerVisualizers_MT())[%s / 16384][%s %% 16384])->AnsiName' % (IndexVal, IndexVal)
+        FNameRef = valobj.CreateValueFromExpression('%s' % IndexVal, Expr)
         assert FNameRef != None
         Val = FNameRef.GetSummary()
-        if IndexVal == 0:
-            return 'name=' + Val + '_' + str(IndexVal-1)
+        if NumberVal != 0:
+            return 'name=%s_%s' % (Val, NumberVal-1)
         else:
-            return 'name=' + Val
+            return 'name=%s' % Val
 
 def UE4UObjectBaseSummaryProvider(valobj,dict):
     Name = valobj.GetChildMemberWithName('Name')
@@ -106,13 +110,13 @@ class UE4TWeakObjectPtrSynthProvider:
     
     def get_child_at_index(self,index):
         logger = lldb.formatters.Logger.Logger()
-        logger >> "Retrieving child " + str(index)
+        logger >> "Retrieving child %s" % index
         if self.ObjectSerialNumberVal >= 1:
-            Expr = 'GObjectArrayForDebugVisualizers->Objects['+str(self.ObjectIndexVal/65536)+']['+str(self.ObjectIndexVal%65536)+'].SerialNumber == '+str(self.ObjectSerialNumberVal)
-            Val = self.valobj.CreateValueFromExpression(str(self.ObjectIndexVal), Expr)
+            Expr = 'GObjectArrayForDebugVisualizers->Objects[%s][%s].SerialNumber == %s' % (self.ObjectIndexVal/65536, self.ObjectIndexVal%65536, self.ObjectSerialNumberVal)
+            Val = self.valobj.CreateValueFromExpression("%s" % self.ObjectIndexVal, Expr)
             Value = Val.GetValueAsUnsigned(0)
             if Value != 0:
-                Expr = 'GObjectArrayForDebugVisualizers->Objects['+str(self.ObjectIndexVal/65536)+']['+str(self.ObjectIndexVal%65536)+'].Object'
+                Expr = 'GObjectArrayForDebugVisualizers->Objects[%s][%s].Object' % (self.ObjectIndexVal/65536, self.ObjectIndexVal%65536)
                 return self.valobj.CreateValueFromExpression('Object', Expr)
             else:
                 Expr = '(void*)0xDEADBEEF'
@@ -141,14 +145,14 @@ def UE4FWeakObjectPtrSummaryProvider(valobj,dict):
         return 'object=nullptr'
     ObjectIndex = valobj.GetChildMemberWithName('ObjectIndex')
     ObjectIndexVal = ObjectIndex.GetValueAsSigned(0)
-    Expr = 'GObjectArrayForDebugVisualizers->Objects['+str(ObjectIndexVal/65536)+']Objects['+str(ObjectIndexVal%65536)+'].SerialNumber == '+str(ObjectSerialNumberVal)
-    Val = valobj.CreateValueFromExpression(str(ObjectIndexVal), Expr)
+    Expr = 'GObjectArrayForDebugVisualizers->Objects[%s][%s].SerialNumber == %s' % (ObjectIndexVal/65536, ObjectIndexVal%65536, ObjectSerialNumberVal)
+    Val = valobj.CreateValueFromExpression('%s' % ObjectIndexVal, Expr)
     ValRef = Val.GetValueAsUnsigned(0)
     if ValRef == 0:
         return 'object=STALE'
     else:
-        Expr = 'GObjectArrayForDebugVisualizers->Objects['+str(ObjectIndexVal/65536)+']['+str(ObjectIndexVal%65536)+'].Object'
-        Val = valobj.CreateValueFromExpression(str(ObjectIndexVal), Expr)
+        Expr = 'GObjectArrayForDebugVisualizers->Objects[%s][%s].Object' % (ObjectIndexVal/65536, ObjectIndexVal%65536)
+        Val = valobj.CreateValueFromExpression("%s" % ObjectIndexVal, Expr)
         return 'object=' + Val.GetValue()
 
 class UE4ChunkedArraySynthProvider:
@@ -174,14 +178,14 @@ class UE4ChunkedArraySynthProvider:
     
     def get_child_at_index(self,index):
         logger = lldb.formatters.Logger.Logger()
-        logger >> "Retrieving child " + str(index)
-        Expr = '(unsigned)sizeof('+str(self.valobj.GetType().GetUnqualifiedType().GetName())+'::FChunk)/'+str(self.ElementTypeSize)
-        self.ChunkBytes = self.valobj.CreateValueFromExpression('['+str(index)+']', Expr)
+        logger >> "Retrieving child %s" % index
+        Expr = '(unsigned)sizeof(%s::FChunk)/%s' % (self.valobj.GetType().GetUnqualifiedType().GetName(), self.ElementTypeSize)
+        self.ChunkBytes = self.valobj.CreateValueFromExpression('[%s]' % index, Expr)
         self.ChunkBytesSize = self.ChunkBytes.GetValueAsUnsigned(0)
         assert self.ChunkBytesSize != 0
         
-        Expr = '*(*((('+str(self.ElementType.GetName())+'**)'+str(self.AllocatorData.GetValue())+')+'+str(index / self.ChunkBytesSize)+')+'+str(index % self.ChunkBytesSize)+')'
-        return self.valobj.CreateValueFromExpression('['+str(index)+']', Expr)
+        Expr = '*(*(((%s**)%s)+%s)+%s)' % (self.ElementType.GetName(), self.AllocatorData.GetValue(), index / self.ChunkBytesSize, index % self.ChunkBytesSize)
+        return self.valobj.CreateValueFromExpression('[%s]' % index, Expr)
     
     def extract_type(self):
         logger = lldb.formatters.Logger.Logger()
@@ -214,7 +218,7 @@ class UE4ChunkedArraySynthProvider:
         return True
 
 def UE4ChunkedArraySummaryProvider(valobj,dict):
-    return 'size=' + str(valobj.GetNumChildren())
+    return 'size=%s' % valobj.GetNumChildren()
 
 class UE4SparseArraySynthProvider:
     
@@ -240,7 +244,7 @@ class UE4SparseArraySynthProvider:
     
     def get_child_at_index(self,index):
         logger = lldb.formatters.Logger.Logger()
-        logger >> "Retrieving child " + str(index)
+        logger >> "Retrieving child %s" % index
         if index < 0:
             return None;
         
@@ -249,15 +253,15 @@ class UE4SparseArraySynthProvider:
         
         Val = None
         if self.SecondaryDataDataVal > 0:
-            Expr = '(bool)((((int*)'+str(self.SecondaryDataData.GetAddress())+')['+str(index)+'/32] >> '+str(index)+') & 1)'
-            Val = self.SecondaryDataData.CreateValueFromExpression('['+str(index)+']', Expr)
+            Expr = '(bool)((((int*)%s)[%s/32] >> %s) & 1)' % (self.SecondaryDataData.GetAddress(), index, index)
+            Val = self.SecondaryDataData.CreateValueFromExpression('[%s]' % index, Expr)
         else:
-            Expr = '(bool)((((int*)'+str(self.InlineData.GetAddress())+')['+str(index)+'/32] >> '+str(index)+') & 1)'
-            Val = self.InlineData.CreateValueFromExpression('['+str(index)+']', Expr)
+            Expr = '(bool)((((int*)%s)[%s/32] >> %s) & 1)' % (self.InlineData.GetAddress(), index, index)
+            Val = self.InlineData.CreateValueFromExpression('[5s]' % index, Expr)
 
         if Val.GetValueAsSigned(0) != 0:
             offset = index * self.ElementTypeSize
-            return self.AllocatorData.CreateChildAtOffset('['+str(index)+']',offset,self.ElementType)
+            return self.AllocatorData.CreateChildAtOffset('[%s]' % index, offset, self.ElementType)
         else:
             return None
                 
@@ -296,7 +300,7 @@ class UE4SparseArraySynthProvider:
         return True
 
 def UE4SparseArraySummaryProvider(valobj,dict):
-    return 'size=' + str(valobj.GetNumChildren())
+    return 'size=%s' % valobj.GetNumChildren()
 
 class UE4BitArraySynthProvider:
    
@@ -321,7 +325,7 @@ class UE4BitArraySynthProvider:
 
     def get_child_at_index(self,index):
         logger = lldb.formatters.Logger.Logger()
-        logger >> "Retrieving child " + str(index)
+        logger >> "Retrieving child %s" % index
         if index < 0:
             return None;
 
@@ -329,11 +333,11 @@ class UE4BitArraySynthProvider:
             return None;
 
         if self.SecondaryDataDataVal > 0:
-            Expr = '(bool)((((int*)'+str(self.SecondaryDataData.GetAddress())+')['+str(index)+'/32] >> '+str(index)+') & 1)'
-            return self.SecondaryDataData.CreateValueFromExpression('['+str(index)+']', Expr)
+            Expr = '(bool)((((int*)%s)[%s/32] >> %s) & 1)' % (self.SecondaryDataData.GetAddress(), index, index)
+            return self.SecondaryDataData.CreateValueFromExpression('[%s]' % index, Expr)
         else:
-            Expr = '(bool)((((int*)'+str(self.InlineData.GetAddress())+')['+str(index)+'/32] >> '+str(index)+') & 1)'
-            return self.InlineData.CreateValueFromExpression('['+str(index)+']', Expr)
+            Expr = '(bool)((((int*)%s)[%s/32] >> %s) & 1)' % (self.InlineData.GetAddress(), index, index)
+            return self.InlineData.CreateValueFromExpression('[%s]' % index, Expr)
 
         return None
 
@@ -356,7 +360,7 @@ class UE4BitArraySynthProvider:
         return True
 
 def UE4BitArraySummaryProvider(valobj,dict):
-    return 'size=' + str(valobj.GetNumChildren())
+    return 'size=%s' % valobj.GetNumChildren()
 
 class UE4ArraySynthProvider:
     
@@ -381,7 +385,7 @@ class UE4ArraySynthProvider:
     
     def get_child_at_index(self,index):
         logger = lldb.formatters.Logger.Logger()
-        logger >> "Retrieving child " + str(index)
+        logger >> "Retrieving child %s" % index
         if index < 0:
             return None;
         
@@ -395,11 +399,11 @@ class UE4ArraySynthProvider:
         try:
             offset = index * self.ElementTypeSize
             if self.Data != None:
-                return self.Data.CreateChildAtOffset('['+str(index)+']',offset,self.ElementType)
+                return self.Data.CreateChildAtOffset('[%s]' % index, offset, self.ElementType)
             elif self.SecondaryDataDataVal > 0:
-                return self.SecondaryDataData.CreateChildAtOffset('['+str(index)+']',offset,self.ElementType)
+                return self.SecondaryDataData.CreateChildAtOffset('[%s]' % index, offset, self.ElementType)
             else:
-                return self.InlineData.CreateChildAtOffset('['+str(index)+']',offset,self.ElementType)
+                return self.InlineData.CreateChildAtOffset('[%s]' % index, offset, self.ElementType)
         except:
             return None
 
@@ -450,20 +454,24 @@ class UE4ArraySynthProvider:
         return True
 
 def UE4ArraySummaryProvider(valobj,dict):
-    return 'size=' + str(valobj.GetChildMemberWithName('ArrayNum').GetValueAsSigned(0))
+    return 'size=%s' % valobj.GetChildMemberWithName('ArrayNum').GetValueAsSigned(0)
 
 class UE4SetSynthProvider:
     
     def __init__(self, valobj, dict):
         logger = lldb.formatters.Logger.Logger()
         self.valobj = valobj
+
+        # Can't cast to TSetElement - the template instantiation won't allow it
+        Expr = '(size_t)sizeof(FSetElementId) + sizeof(int32)'
+        self.TSetElement = self.valobj.CreateValueFromExpression('TSetElement', Expr)
     
     def num_children(self):
         logger = lldb.formatters.Logger.Logger()
         try:
-            self.ArrayNumVal = self.ArrayNum.GetValueAsUnsigned(0)
-            self.NumFreeIndicesVal = self.NumFreeIndices.GetValueAsUnsigned(0)
-            return self.ArrayNumVal - self.NumFreeIndicesVal;
+            ArrayNumVal = self.ArrayNum.GetValueAsUnsigned(0)
+            NumFreeIndicesVal = self.NumFreeIndices.GetValueAsUnsigned(0)
+            return ArrayNumVal - NumFreeIndicesVal;
         except:
             return 0;
     
@@ -476,26 +484,26 @@ class UE4SetSynthProvider:
     
     def get_child_at_index(self,index):
         logger = lldb.formatters.Logger.Logger()
-        logger >> "Retrieving child " + str(index)
+        logger >> "Retrieving child %s" % index
         if index < 0:
             return None;
-        
+
         if index >= self.num_children():
             return None;
         try:
             offset = index * self.ElementTypeSize
             HasObject = 0
             if self.SecondaryDataDataVal > 0:
-                Expr = '(bool)((((int*)'+str(self.SecondaryDataDataVal)+')['+str(index)+'/32] >> '+str(index)+') & 1)'
-                HasObject = 1 ##self.AllocationFlagsSecondaryDataData.CreateValueFromExpression('['+str(index)+']', Expr).GetValueAsUnsigned(0)
+                Expr = '(bool)((((int*)%s)[%s/32] >> %s) & 1)' % (self.SecondaryDataDataVal, index, index)
+                HasObject = 1 ##self.AllocationFlagsSecondaryDataData.CreateValueFromExpression('[%s]' % index, Expr).GetValueAsUnsigned(0)
             else:
-                Expr = '(bool)((((int*)'+str(self.AllocationFlagsInlineDataAddr)+')['+str(index)+'/32] >> '+str(index)+') & 1)'
-                HasObject = 1 ##self.AllocationFlagsInlineData.CreateValueFromExpression('['+str(index)+']', Expr).GetValueAsUnsigned(0)
+                Expr = '(bool)((((int*)%s)[%s/32] >> %s) & 1)' % (self.AllocationFlagsInlineDataAddr, index, index)
+                HasObject = 1 ##self.AllocationFlagsInlineData.CreateValueFromExpression('[%s]' % index, Expr).GetValueAsUnsigned(0)
 
             if HasObject == 1:
-                return self.AllocatorInstanceData.CreateChildAtOffset('['+str(index)+']',offset,self.ElementType)
+                return self.AllocatorInstanceData.CreateChildAtOffset('[%s]' % index, offset, self.ElementType)
             else:
-                return self.valobj.CreateValueFromExpression('['+str(index)+']', '(void*)0xDEADBEEF')
+                return self.valobj.CreateValueFromExpression('[%s]' % index, '(void*)0xDEADBEEF')
         except:
             return None
     
@@ -529,11 +537,8 @@ class UE4SetSynthProvider:
             self.AllocationFlagsSecondaryDataData = self.AllocationFlagsSecondaryData.GetChildMemberWithName('Data')
             self.SecondaryDataDataVal = self.AllocationFlagsSecondaryDataData.GetValueAsUnsigned(0)
             self.ElementType = self.extract_type()
-            # Can't cast to TSetElement - the template instantiation won't allow it
-            Expr = '(size_t)sizeof(FSetElementId) + sizeof(int32)'
-            TSetElement = self.valobj.CreateValueFromExpression('TSetElement', Expr)
             # This may fail due to C++ struct padding - will have to check
-            self.ElementTypeSize = self.ElementType.GetByteSize() + TSetElement.GetValueAsUnsigned(0)
+            self.ElementTypeSize = self.ElementType.GetByteSize() + self.TSetElement.GetValueAsUnsigned(0)
             assert self.ElementTypeSize != 0
         except:
             pass
@@ -542,20 +547,24 @@ class UE4SetSynthProvider:
         return True
 
 def UE4SetSummaryProvider(valobj,dict):
-    return 'size=' + str(valobj.GetNumChildren())
+    return 'size=%s' % valobj.GetNumChildren()
 
 class UE4MapSynthProvider:
     
     def __init__(self, valobj, dict):
         logger = lldb.formatters.Logger.Logger()
         self.valobj = valobj
+
+        # Can't cast to TSetElement - the template instantiation won't allow it
+        Expr = '(size_t)sizeof(FSetElementId) + sizeof(int32)'
+        self.TSetElement = self.valobj.CreateValueFromExpression('TSetElement', Expr)
     
     def num_children(self):
         logger = lldb.formatters.Logger.Logger()
         try:
-            self.ArrayNumVal = self.ArrayNum.GetValueAsUnsigned(0)
-            self.NumFreeIndicesVal = self.NumFreeIndices.GetValueAsUnsigned(0)
-            return self.ArrayNumVal - self.NumFreeIndicesVal;
+            ArrayNumVal = self.ArrayNum.GetValueAsUnsigned(0)
+            NumFreeIndicesVal = self.NumFreeIndices.GetValueAsUnsigned(0)
+            return ArrayNumVal - NumFreeIndicesVal;
         except:
             return 0;
     
@@ -568,7 +577,7 @@ class UE4MapSynthProvider:
     
     def get_child_at_index(self,index):
         logger = lldb.formatters.Logger.Logger()
-        logger >> "Retrieving child " + str(index)
+        logger >> "Retrieving child %s" % index
         if index < 0:
             return None;
         
@@ -578,16 +587,16 @@ class UE4MapSynthProvider:
             offset = index * self.ElementTypeSize
             HasObject = 0
             if self.SecondaryDataDataVal != 0:
-                Expr = '(bool)((((unsigned int*)'+str(self.SecondaryDataDataVal)+')['+str(index)+'/32] >> '+str(index)+') & 1)'
-                HasObject = 1 ##self.AllocationFlagsSecondaryDataData.CreateValueFromExpression('['+str(index)+']', Expr).GetValueAsUnsigned(0)
+                Expr = '(bool)((((unsigned int*)%s)[%s/32] >> %s) & 1)' % (self.SecondaryDataDataVal, index, index)
+                HasObject = 1 ##self.AllocationFlagsSecondaryDataData.CreateValueFromExpression('[%s]' % index, Expr).GetValueAsUnsigned(0)
             else:
-                Expr = '(bool)((((unsigned int*)'+str(self.AllocationFlagsInlineDataAddr)+')['+str(index)+'/32] >> '+str(index)+') & 1)'
-                HasObject = 1 ##self.AllocationFlagsInlineData.CreateValueFromExpression('['+str(index)+']', Expr).GetValueAsUnsigned(0)
+                Expr = '(bool)((((unsigned int*)%s)[%s/32] >> %s) & 1)' % (self.AllocationFlagsInlineDataAddr, index, index)
+                HasObject = 1 ##self.AllocationFlagsInlineData.CreateValueFromExpression('[%s]' % index, Expr).GetValueAsUnsigned(0)
             
             if HasObject == 1:
-                return self.AllocatorInstanceData.CreateChildAtOffset('['+str(index)+']',offset,self.ElementType)
+                return self.AllocatorInstanceData.CreateChildAtOffset('[%s]' % index,offset,self.ElementType)
             else:
-                return self.valobj.CreateValueFromExpression('['+str(index)+']', '(void*)0xDEADBEEF')
+                return self.valobj.CreateValueFromExpression('[%s]' % index, '(void*)0xDEADBEEF')
         except:
             return None
     
@@ -622,11 +631,8 @@ class UE4MapSynthProvider:
             self.AllocationFlagsSecondaryDataData = self.AllocationFlagsSecondaryData.GetChildMemberWithName('Data')
             self.SecondaryDataDataVal = self.AllocationFlagsSecondaryDataData.GetValueAsUnsigned(0)
             self.ElementType = self.extract_type()
-            # Can't cast to TSetElement - the template instantiation won't allow it
-            Expr = '(size_t)sizeof(FSetElementId) + sizeof(int32)'
-            TSetElement = self.valobj.CreateValueFromExpression('TSetElement', Expr)
             # This may fail due to C++ struct padding - will have to check
-            self.ElementTypeSize = self.ElementType.GetByteSize() + TSetElement.GetValueAsUnsigned(0)
+            self.ElementTypeSize = self.ElementType.GetByteSize() + self.TSetElement.GetValueAsUnsigned(0)
             assert self.ElementTypeSize != 0
         except:
             pass
@@ -635,13 +641,13 @@ class UE4MapSynthProvider:
         return True
 
 def UE4MapSummaryProvider(valobj,dict):
-    return 'size=' + str(valobj.GetNumChildren())
+    return 'size=%s' % valobj.GetNumChildren()
 
 def __lldb_init_module(debugger,dict):
     debugger.HandleCommand('type summary add -F UE4DataFormatters_2ByteChars.UE4TCharSummaryProvider -e TCHAR -w UE4DataFormatters')
     debugger.HandleCommand('type summary add -F UE4DataFormatters_2ByteChars.UE4FStringSummaryProvider -e -x "FString$" -w UE4DataFormatters')
     debugger.HandleCommand('type summary add -F UE4DataFormatters_2ByteChars.UE4FNameSummaryProvider -e -x "FName$" -w UE4DataFormatters')
-    debugger.HandleCommand('type summary add -F UE4DataFormatters_2ByteChars.UE4FMinimalNameSummaryProvider -e -x "UE4FMinimalName$" -w UE4DataFormatters')
+    debugger.HandleCommand('type summary add -F UE4DataFormatters_2ByteChars.UE4FMinimalNameSummaryProvider -e -x "FMinimalName$" -w UE4DataFormatters')
     debugger.HandleCommand('type summary add -F UE4DataFormatters_2ByteChars.UE4UObjectBaseSummaryProvider -e UObject -w UE4DataFormatters')
     debugger.HandleCommand('type summary add -F UE4DataFormatters_2ByteChars.UE4UObjectBaseSummaryProvider -e UObjectBase -w UE4DataFormatters')
     debugger.HandleCommand('type summary add -F UE4DataFormatters_2ByteChars.UE4UObjectBaseSummaryProvider -e UObjectBaseUtility -w UE4DataFormatters')
