Fix invalid string index in lldb script.

diff --git a/Engine/Extras/LLDBDataFormatters/UE4DataFormatters.py b/Engine/Extras/LLDBDataFormatters/UE4DataFormatters.py
--- a/Engine/Extras/LLDBDataFormatters/UE4DataFormatters.py
+++ b/Engine/Extras/LLDBDataFormatters/UE4DataFormatters.py
@@ -66,7 +66,7 @@ def UE4FNameSummaryProvider(valobj,dict):
         FNameRef = valobj.CreateValueFromExpression(str(IndexVal), Expr)
         assert FNameRef != None
         Val = FNameRef.GetSummary()
-        if IndexVal == 0:
+        if IndexVal != 0:
             return 'name=' + Val + '_' + str(IndexVal-1)
         else:
             return 'name=' + Val
