Updated debugger visualizers to support FProperty related classes

diff --git a/Engine/Extras/LLDBDataFormatters/UE4DataFormatters.py b/Engine/Extras/LLDBDataFormatters/UE4DataFormatters.py
--- a/Engine/Extras/LLDBDataFormatters/UE4DataFormatters.py
+++ b/Engine/Extras/LLDBDataFormatters/UE4DataFormatters.py
@@ -90,6 +90,14 @@ def UE4UObjectBaseSummaryProvider(valobj,dict):
     Name = valobj.GetChildMemberWithName('Name')
     return Name.GetSummary()
 
+def UE4FFieldClassSummaryProvider(valobj,dict):
+    Name = valobj.GetChildMemberWithName('Name')
+    return Name.GetSummary()
+
+def UE4FFieldSummaryProvider(valobj,dict):
+    Name = valobj.GetChildMemberWithName('NamePrivate')
+    return Name.GetSummary()
+
 class UE4TWeakObjectPtrSynthProvider:
     
     def __init__(self, valobj, dict):
@@ -652,6 +660,8 @@ def __lldb_init_module(debugger,dict):
     debugger.HandleCommand('type summary add -F UE4DataFormatters.UE4UObjectBaseSummaryProvider -e UObject -w UE4DataFormatters')
     debugger.HandleCommand('type summary add -F UE4DataFormatters.UE4UObjectBaseSummaryProvider -e UObjectBase -w UE4DataFormatters')
     debugger.HandleCommand('type summary add -F UE4DataFormatters.UE4UObjectBaseSummaryProvider -e UObjectBaseUtility -w UE4DataFormatters')
+    debugger.HandleCommand('type summary add -F UE4DataFormatters.UE4FFieldClassSummaryProvider -e FFieldClass -w UE4DataFormatters')
+    debugger.HandleCommand('type summary add -F UE4DataFormatters.UE4FFieldSummaryProvider -e FField -w UE4DataFormatters')
     debugger.HandleCommand('type summary add -F UE4DataFormatters.UE4FWeakObjectPtrSummaryProvider -e FWeakObjectPtr -w UE4DataFormatters')
     debugger.HandleCommand('type synthetic add -l UE4DataFormatters.UE4TWeakObjectPtrSynthProvider -x "TWeakObjectPtr<.+>$" -w UE4DataFormatters')
     debugger.HandleCommand('type synthetic add -l UE4DataFormatters.UE4TWeakObjectPtrSynthProvider -x "TAutoWeakObjectPtr<.+>$" -w UE4DataFormatters')
diff --git a/Engine/Extras/LLDBDataFormatters/UE4DataFormatters_2ByteChars.py b/Engine/Extras/LLDBDataFormatters/UE4DataFormatters_2ByteChars.py
--- a/Engine/Extras/LLDBDataFormatters/UE4DataFormatters_2ByteChars.py
+++ b/Engine/Extras/LLDBDataFormatters/UE4DataFormatters_2ByteChars.py
@@ -94,6 +94,14 @@ def UE4UObjectBaseSummaryProvider(valobj,dict):
     Name = valobj.GetChildMemberWithName('Name')
     return Name.GetSummary()
 
+def UE4FFieldClassSummaryProvider(valobj,dict):
+    Name = valobj.GetChildMemberWithName('Name')
+    return Name.GetSummary()
+
+def UE4FFieldSummaryProvider(valobj,dict):
+    Name = valobj.GetChildMemberWithName('NamePrivate')
+    return Name.GetSummary()
+
 class UE4TWeakObjectPtrSynthProvider:
     
     def __init__(self, valobj, dict):
@@ -651,6 +659,8 @@ def __lldb_init_module(debugger,dict):
     debugger.HandleCommand('type summary add -F UE4DataFormatters_2ByteChars.UE4UObjectBaseSummaryProvider -e UObject -w UE4DataFormatters')
     debugger.HandleCommand('type summary add -F UE4DataFormatters_2ByteChars.UE4UObjectBaseSummaryProvider -e UObjectBase -w UE4DataFormatters')
     debugger.HandleCommand('type summary add -F UE4DataFormatters_2ByteChars.UE4UObjectBaseSummaryProvider -e UObjectBaseUtility -w UE4DataFormatters')
+    debugger.HandleCommand('type summary add -F UE4DataFormatters.UE4FFieldClassSummaryProvider -e FFieldClass -w UE4DataFormatters')
+    debugger.HandleCommand('type summary add -F UE4DataFormatters.UE4FFieldSummaryProvider -e FField -w UE4DataFormatters')
     debugger.HandleCommand('type summary add -F UE4DataFormatters_2ByteChars.UE4FWeakObjectPtrSummaryProvider -e FWeakObjectPtr -w UE4DataFormatters')
     debugger.HandleCommand('type synthetic add -l UE4DataFormatters_2ByteChars.UE4TWeakObjectPtrSynthProvider -x "TWeakObjectPtr<.+>$" -w UE4DataFormatters')
     debugger.HandleCommand('type synthetic add -l UE4DataFormatters_2ByteChars.UE4TWeakObjectPtrSynthProvider -x "TAutoWeakObjectPtr<.+>$" -w UE4DataFormatters')
