Rename and fix up UE4 referencing in GDB/LLDB formatters

diff --git a/Engine/Build/BatchFiles/Linux/GenerateGDBInit.sh b/Engine/Build/BatchFiles/Linux/GenerateGDBInit.sh
--- a/Engine/Build/BatchFiles/Linux/GenerateGDBInit.sh
+++ b/Engine/Build/BatchFiles/Linux/GenerateGDBInit.sh
@@ -6,32 +6,32 @@ TOP_DIR=$(cd "$SCRIPT_DIR/../../.." ; pwd)
 
 AddGDBPrettyPrinters()
 {
-	echo -ne "Attempting to set up UE4 pretty printers for gdb (existing UE4Printers.py, if any, will be overwritten)...\n\t"
+	echo -ne "Attempting to set up UE pretty printers for gdb (existing UEPrinters.py, if any, will be overwritten)...\n\t"
 
 	# Copy the pretty printer into the appropriate folder.
 	mkdir -p ~/.config/Epic/GDBPrinters/
-	if [ -e ~/.config/Epic/GDBPrinters/UE4Printers.py ]; then
-		chmod 644 ~/.config/Epic/GDBPrinters/UE4Printers.py 	# set R/W so we can overwrite it
+	if [ -e ~/.config/Epic/GDBPrinters/UEPrinters.py ]; then
+		chmod 644 ~/.config/Epic/GDBPrinters/UEPrinters.py 	# set R/W so we can overwrite it
 	fi
-	cp "$TOP_DIR/Extras/GDBPrinters/UE4Printers.py" ~/.config/Epic/GDBPrinters/
-	echo -ne "updated UE4Printers.py\n\t"
-	chmod 644 ~/.config/Epic/GDBPrinters/UE4Printers.py 	# set R/W again (it can be read-only if copied from Perforce)
+	cp "$TOP_DIR/Extras/GDBPrinters/UEPrinters.py" ~/.config/Epic/GDBPrinters/
+	echo -ne "updated UEPrinters.py\n\t"
+	chmod 644 ~/.config/Epic/GDBPrinters/UEPrinters.py 	# set R/W again (it can be read-only if copied from Perforce)
 
 	# Check if .gdbinit exists. If not create else add needed parts.
 	if [ ! -f ~/.gdbinit ]; then
 		echo "no ~/.gdbinit file found - creating a new one."
-		echo -e "python \nimport sys\n\nsys.path.append('$HOME/.config/Epic/GDBPrinters/')\n\nfrom UE4Printers import register_ue4_printers\nregister_ue4_printers(None)\nprint(\"Registered pretty printers for UE4 classes\")\n\nend" >> ~/.gdbinit
+		echo -e "python \nimport sys\n\nsys.path.append('$HOME/.config/Epic/GDBPrinters/')\n\nfrom UEPrinters import register_ue_printers\nregister_ue_printers(None)\nprint(\"Registered pretty printers for UE classes\")\n\nend" >> ~/.gdbinit
 	else
-		if grep -q "register_ue4_printers" ~/.gdbinit; then
+		if grep -q "register_ue_printers" ~/.gdbinit; then
 			echo "found necessary entries in ~/.gdbinit file, not changing it."
 		else
 			echo -e "cannot modify .gdbinit. Please add the below lines manually:\n\n"
 			echo -e "python"
 			echo -e "\timport sys"
 			echo -e "\tsys.path.append('$HOME/.config/Epic/GDBPrinters/')"
-			echo -e "\tfrom UE4Printers import register_ue4_printers"
-			echo -e "\tregister_ue4_printers(None)"
-			echo -e "\tprint(\"Registered pretty printers for UE4 classes\")"
+			echo -e "\tfrom UEPrinters import register_ue_printers"
+			echo -e "\tregister_ue_printers(None)"
+			echo -e "\tprint(\"Registered pretty printers for UE classes\")"
 			echo -e "\tend"
 			echo -e "\n\n"
 		fi
diff --git a/Engine/Build/BatchFiles/Mac/GenerateLLDBInit.sh b/Engine/Build/BatchFiles/Mac/GenerateLLDBInit.sh
--- a/Engine/Build/BatchFiles/Mac/GenerateLLDBInit.sh
+++ b/Engine/Build/BatchFiles/Mac/GenerateLLDBInit.sh
@@ -2,5 +2,5 @@
 
 if [ ! -f ~/.lldbinit ]; then
   echo 'echo -n "settings set target.inline-breakpoint-strategy always\n"' | bash > ~/.lldbinit
-  echo 'echo -n "command script import \"`pwd`/../../../Extras/LLDBDataFormatters/UE4DataFormatters_2ByteChars.py\"\n"' | bash >> ~/.lldbinit
+  echo 'echo -n "command script import \"`pwd`/../../../Extras/LLDBDataFormatters/UEDataFormatters_2ByteChars.py\"\n"' | bash >> ~/.lldbinit
 fi
diff --git a/Engine/Build/InstalledEngineFilters.xml b/Engine/Build/InstalledEngineFilters.xml
--- a/Engine/Build/InstalledEngineFilters.xml
+++ b/Engine/Build/InstalledEngineFilters.xml
@@ -329,7 +329,7 @@
 
 	<Property Name="CopyEditorFilterMac" Value="$(CopyEditorFilter);$(MacBuildFiles)">
 		<!-- Extra tools -->
-		Engine/Extras/LLDBDataFormatters/UE4DataFormatters.py
+		Engine/Extras/LLDBDataFormatters/UEDataFormatters.py
 	</Property>
 
 	<Property Name="CopyEditorExceptionsMac" Value="$(CopyEditorExceptions)">
@@ -365,7 +365,7 @@
 		Engine/Binaries/ThirdParty/Mono/...
 
 		<!-- Extra tools -->
-		Engine/Extras/GDBPrinters/UE4Printers.py
+		Engine/Extras/GDBPrinters/UEPrinters.py
 
 		<!-- Bundled libcxx -->
 		Engine/Source/ThirdParty/Linux/LibCxx/...
diff --git a/Engine/Extras/GDBPrinters/UE4Printers.py b/Engine/Extras/GDBPrinters/UE4Printers.py
deleted file mode 100644
diff --git a/Engine/Extras/GDBPrinters/UEPrinters.py b/Engine/Extras/GDBPrinters/UEPrinters.py
new file mode 100644
--- /dev/null
+++ b/Engine/Extras/GDBPrinters/UEPrinters.py
@@ -0,0 +1,728 @@
+#
+# GDB Printers for the Unreal Engine 4
+#
+# How to install:
+# If the file ~/.gdbinit doesn't exist
+#        touch ~/.gdbinit
+#        open ~/.gdbinit
+#
+# and add the following lines:
+#   python
+#   import sys
+#   ...
+#   sys.path.insert(0, '/Path/To/Epic/UE/Engine/Extras/GDBPrinters')      <--
+#   ...
+#   from UEPrinters import register_ue_printers                          <--
+#   register_ue_printers (None)                                           <--
+#   ...
+#   end
+
+
+import gdb
+import itertools
+import re
+import sys
+
+# ------------------------------------------------------------------------------
+# We make our own base of the iterator to prevent issues between Python 2/3.
+#
+
+if sys.version_info[0] == 3:
+	Iterator = object
+else:
+	class Iterator(object):
+
+		def next(self):
+			return type(self).__next__(self)
+
+def default_iterator(val):
+	for field in val.type.fields():
+		yield field.name, val[field.name]
+
+
+# ------------------------------------------------------------------------------
+#
+#  Custom pretty printers.
+#
+#
+
+
+# ------------------------------------------------------------------------------
+# FBitReference
+#
+class FBitReferencePrinter:
+	def __init__(self, typename, val):
+		self.Value = val
+
+	def to_string(self):
+		self.Mask = self.Value['Mask']
+		self.Data = self.Value['Data']
+		return '\'%d\'' % (self.Data & self.Mask)
+
+# ------------------------------------------------------------------------------
+# TBitArray
+#
+class TBitArrayPrinter:
+	"Print TBitArray"
+
+	class _iterator(Iterator):
+		def __init__(self, val):
+			self.Value = val
+			self.Counter = -1
+
+			try:
+				self.NumBits = self.Value['NumBits']
+				if self.NumBits.is_optimized_out:
+					self.NumBits = 0
+				else:
+					self.AllocatorInstance = self.Value['AllocatorInstance']
+					self.InlineData = self.AllocatorInstance['InlineData']
+					self.SecondaryData = self.AllocatorInstance['SecondaryData']
+					self.SecondaryDataData = self.AllocatorInstance['SecondaryData']['Data']
+					if self.SecondaryData != None:
+						self.SecondaryDataData = self.SecondaryData['Data']
+			except:
+				raise
+
+		def __iter__(self):
+			return self
+
+		def __next__(self):
+			if self.NumBits == 0:
+				raise StopIteration
+
+			self.Counter = self.Counter + 1
+
+			if self.Counter >= self.NumBits:
+				raise StopIteration
+
+			if self.SecondaryDataData > 0:
+				data = self.SecondaryDataData.cast(gdb.lookup_type("uint32").pointer())
+			else:
+				data = self.InlineData.cast(gdb.lookup_type("uint32").pointer())
+
+			return ('[%d]' % self.Counter, (data[self.Counter/32] >> self.Counter) & 1)
+
+	def __init__(self, typename, val):
+		self.Value = val
+		self.NumBits = self.Value['NumBits']
+
+	def to_string(self):
+		if self.NumBits.is_optimized_out:
+			pass
+		if self.NumBits == 0:
+			return 'empty'
+		pass
+
+	def children(self):
+		return self._iterator(self.Value)
+
+	def display_hint(self):
+		return 'array'
+# ------------------------------------------------------------------------------
+# TIndirectArray
+#
+
+
+# ------------------------------------------------------------------------------
+# TChunkedArray
+#
+class TChunkedArrayPrinter:
+	"Print TChunkedArray"
+
+	class _iterator(Iterator):
+		def __init__(self, val, typename):
+			self.Value = val
+			self.Typename = typename
+			self.Counter = -1
+			self.ElementType = self.Value.type.template_argument(0)
+			self.ElementTypeSize = self.ElementType.sizeof
+
+			try:
+				self.NumElements = self.Value['NumElements']
+				if self.NumElements.is_optimized_out:
+					self.NumElements = 0
+				else:
+					self.Chunks = self.Value['Chunks']
+					self.Array = self.Chunks['Array']
+					self.ArrayNum = self.Array['ArrayNum']
+					self.AllocatorInstance = self.Array['AllocatorInstance']
+					self.AllocatorData = self.AllocatorInstance['Data']
+			except:
+				raise
+
+		def __iter__(self):
+			return self
+
+		def __next__(self):
+			return self.next()
+
+		def __next__(self):
+			if self.NumElements == 0:
+				raise StopIteration
+
+			self.Counter = self.Counter + 1
+
+			if self.Counter >= self.NumElements:
+				raise StopIteration()
+
+			Expr = '(unsigned)sizeof('+str(self.Typename)+'::FChunk)/'+str(self.ElementTypeSize)
+			self.ChunkBytes = gdb.parse_and_eval(Expr)
+			assert self.ChunkBytes != 0
+
+			Expr = '*(*((('+str(self.ElementType.name)+'**)'+str(self.AllocatorData)+')+'+str(self.Counter / self.ChunkBytes)+')+'+str(self.Counter % self.ChunkBytes)+')'
+			Val = gdb.parse_and_eval(Expr)
+			return ('[%d]' % self.Counter, Val)
+
+	def __init__(self, typename, val):
+		self.Value = val
+		self.Typename = typename
+		self.NumElements = self.Value['NumElements']
+
+	def to_string(self):
+		if self.NumElements.is_optimized_out:
+			pass
+		if self.NumElements == 0:
+			return 'empty'
+		pass
+
+	def children(self):
+		return self._iterator(self.Value, self.Typename)
+
+	def display_hint(self):
+		return 'array'
+
+# ------------------------------------------------------------------------------
+# TSparseArray
+#
+class TSparseArrayPrinter:
+	"Print TSparseArray"
+
+	class _iterator(Iterator):
+		def __init__(self, val, typename):
+
+			self.Value = val
+			self.Counter = -1
+			self.Typename = typename
+			self.ElementType = self.Value.type.template_argument(0)
+			self.ElementTypeSize = self.ElementType.sizeof
+			assert self.ElementTypeSize != 0
+
+			try:
+				self.NumFreeIndices = self.Value['NumFreeIndices']
+				self.Data = self.Value['Data']
+				self.ArrayNum = self.Data['ArrayNum']
+				if self.ArrayNum.is_optimized_out:
+					self.ArrayNum = 0
+				else:
+					self.AllocatorInstance = self.Data['AllocatorInstance']
+					self.AllocatorData = self.AllocatorInstance['Data']
+					self.AllocationFlags = self.Value['AllocationFlags']
+					self.AllocationFlagsInstance = self.AllocationFlags['AllocatorInstance']
+					self.InlineData = self.AllocationFlagsInstance['InlineData']
+					self.SecondaryData = self.AllocationFlagsInstance['SecondaryData']
+					self.SecondaryDataData = self.SecondaryData['Data']
+			except:
+				raise
+
+		def __iter__(self):
+			return self
+
+		def __next__(self):
+			return self.next()
+
+		def __next__(self):
+			if self.ArrayNum == 0:
+				raise StopIteration
+
+			self.Counter = self.Counter + 1
+
+			if self.Counter >= self.ArrayNum:
+				raise StopIteration
+			else:
+				Data = None
+				if self.SecondaryDataData > 0:
+					Data = (self.SecondaryDataData.address.cast(gdb.lookup_type("int").pointer())[self.Counter/32] >> self.Counter) & 1
+				else:
+					Data = (self.InlineData.address.cast(gdb.lookup_type("int").pointer())[self.Counter/32] >> self.Counter) & 1
+
+				Value = None
+				if Data != 0:
+					offset = self.Counter * self.ElementTypeSize
+					Value = (self.AllocatorData + offset).cast(self.ElementType.pointer())
+				else:
+					Value = None
+
+				return ('[%s]' % self.Counter, Value.dereference())
+
+	def __init__(self, typename, val):
+		self.Value = val
+		self.Typename = typename
+		self.ArrayNum = self.Value['Data']['ArrayNum']
+
+	def to_string(self):
+		if self.ArrayNum.is_optimized_out:
+			pass
+		if self.ArrayNum == 0:
+			return 'empty'
+		pass
+
+	def children(self):
+		return self._iterator(self.Value, self.Typename)
+
+	def display_hint(self):
+		return 'array'
+
+# ------------------------------------------------------------------------------
+# TSet
+#
+class TSetPrinter:
+	"Print TSet"
+
+	class _iterator(Iterator):
+		def __init__(self, val, typename):
+			self.Value = val
+			self.Counter = -1
+			self.typename = typename
+			self.ElementType = self.Value.type.template_argument(0)
+
+			try:
+				self.Elements = self.Value["Elements"]
+				self.ElementsData = self.Elements["Data"]
+				self.ElementsArrayNum = self.ElementsData['ArrayNum']
+				self.NumFreeIndices = self.Elements['NumFreeIndices']
+
+				self.AllocatorInstance = self.ElementsData['AllocatorInstance']
+				self.AllocatorInstanceData = self.AllocatorInstance['Data']
+
+				self.AllocationFlags = self.Elements['AllocationFlags']
+				self.AllocationFlagsInstance = self.AllocationFlags['AllocatorInstance']
+				self.InlineData = self.AllocationFlagsInstance['InlineData']
+				self.SecondaryData = self.AllocationFlagsInstance['SecondaryData']
+				self.SecondaryDataData = self.SecondaryData['Data']
+			except:
+				self.ElementsArrayNum = 0
+
+			Expr = '(size_t)sizeof(FSetElementId) + sizeof(int32)'
+			TSetElement = gdb.parse_and_eval(Expr)
+			self.ElementTypeSize = self.ElementType.sizeof + TSetElement
+
+		def __iter__(self):
+			return self
+
+		def __next__(self):
+			self.Counter = self.Counter + 1
+
+			if self.Counter >= self.ElementsArrayNum:
+				raise StopIteration()
+			else:
+				Data = None
+				if self.SecondaryDataData > 0:
+					Data = (self.SecondaryDataData.address.cast(gdb.lookup_type("int").pointer())[self.Counter/32] >> self.Counter) & 1
+				else:
+					Data = (self.InlineData.address.cast(gdb.lookup_type("int").pointer())[self.Counter/32] >> self.Counter) & 1
+
+				Value = None
+				if Data != 0:
+					offset = self.Counter * self.ElementTypeSize
+					Value = (self.AllocatorInstanceData + offset).cast(self.ElementType.pointer())
+				else:
+					Value = None
+
+			return ('[%s]' % self.Counter, Value.dereference())
+
+	def __init__(self, typename, val):
+		self.Value = val
+		self.typename = typename
+		self.ArrayNum = self.Value["Elements"]["Data"]['ArrayNum']
+
+	def to_string(self):
+		if self.ArrayNum.is_optimized_out:
+			pass
+		if self.ArrayNum == 0:
+			return 'empty'
+		pass
+
+	def children(self):
+		return self._iterator(self.Value, self.typename)
+
+	def display_hint(self):
+		return 'array'
+
+# ------------------------------------------------------------------------------
+# TSetElementPrinter
+#
+
+class TSetElementPrinter:
+	"Print TSetElement"
+
+	def __init__(self, typename, val):
+		self.Value = val
+
+	def to_string(self):
+		if self.Value.is_optimized_out:
+			return '<optimized out>'
+
+		return self.Value["Value"]
+
+	def display_hint(self):
+		return "string"
+
+# ------------------------------------------------------------------------------
+# TMap
+#
+class TMapPrinter:
+	"Print TMap"
+
+	class _iterator(Iterator):
+		def __init__(self, val):
+			self.Value = val
+			self.Counter = -1
+			try:
+				self.Pairs = self.Value['Pairs']
+				if self.Pairs.is_optimized_out:
+					self.ArrayNum = 0
+				else:
+					self.Elements = self.Pairs['Elements']
+					self.ElementsData = self.Elements['Data']
+					self.ArrayNum = self.ElementsData['ArrayNum']
+			except:
+				raise
+
+		def __iter__(self):
+			return self
+
+		def __next__(self):
+			if self.ArrayNum == 0:
+				raise StopIteration
+
+			self.Counter = self.Counter + 1
+
+			if self.Counter > 0:
+				raise StopIteration
+
+			return ('Pairs', self.Pairs)
+
+	def __init__(self, typename, val):
+		self.Value = val
+		self.ArrayNum = self.Value['Pairs']['Elements']['Data']['ArrayNum']
+
+	def children(self):
+		return self._iterator(self.Value)
+
+	def to_string(self):
+		if self.ArrayNum.is_optimized_out:
+			pass
+		if self.ArrayNum == 0:
+			return 'empty'
+		pass
+
+	def display_hint(self):
+		return 'map'
+
+# ------------------------------------------------------------------------------
+# TWeakObjectPtr
+#
+
+class TWeakObjectPtrPrinter:
+	"Print TWeakObjectPtr"
+
+	class _iterator(Iterator):
+		def __init__(self, val):
+			self.Value = val
+			self.Counter = 0
+			self.Object = None
+
+			self.ObjectSerialNumber = int(self.Value['ObjectSerialNumber'])
+			if self.ObjectSerialNumber >= 1:
+				ObjectIndexValue = int(self.Value['ObjectIndex'])
+				ObjectItemExpr = 'GCoreObjectArrayForDebugVisualizers->Objects['+str(ObjectIndexValue)+'/FChunkedFixedUObjectArray::NumElementsPerChunk]['+str(ObjectIndexValue)+ '% FChunkedFixedUObjectArray::NumElementsPerChunk]'
+				ObjectItem = gdb.parse_and_eval(ObjectItemExpr);
+				IsValidObject = int(ObjectItem['SerialNumber']) == self.ObjectSerialNumber
+				if IsValidObject == True:
+					ObjectType = self.Value.type.template_argument(0)
+					self.Object = ObjectItem['Object'].dereference().cast(ObjectType.reference())
+
+		def __iter__(self):
+			return self
+
+		def __next__(self):
+			if self.Counter > 0:
+				raise StopIteration
+
+			self.Counter = self.Counter + 1
+
+			if self.Object != None:
+				return ('Object', self.Object)
+			elif self.ObjectSerialNumber > 0:
+				return ('Object', 'STALE')
+			else:
+				return ('Object', 'nullptr')
+
+
+	def __init__(self, typename, val):
+		self.Value = val
+
+	def children(self):
+		return self._iterator(self.Value)
+
+	def to_string(self):
+		ObjectType = self.Value.type.template_argument(0)
+		return 'TWeakObjectPtr<%s>' % ObjectType.name;
+
+
+# ------------------------------------------------------------------------------
+# FString
+#
+class FStringPrinter:
+	"Print FString"
+
+	def __init__(self, typename, val):
+		self.Value = val
+
+	def to_string(self):
+		if self.Value.is_optimized_out:
+			return '<optimized out>'
+
+		ArrayNum = self.Value['Data']['ArrayNum']
+		if ArrayNum == 0:
+			return 'empty'
+		elif ArrayNum < 0:
+			return "nullptr"
+		else:
+			ActualData = self.Value['Data']['AllocatorInstance']['Data']
+			data = ActualData.cast(gdb.lookup_type("TCHAR").pointer())
+			return '%s' % (data.string())
+
+	def display_hint (self):
+		return 'string'
+
+
+# ------------------------------------------------------------------------------
+# FNameEntry
+#
+
+class FNameEntryPrinter:
+	"Print FNameEntry"
+
+	def __init__(self, typename, val):
+		self.Value = val
+
+	def to_string(self):
+		self.Header = self.Value['Header']
+		IsWideString = self.Header['bIsWide'].cast(gdb.lookup_type('bool'))
+		self.AnsiName = self.Value['AnsiName']
+		self.WideName = self.Value['WideName']
+		Len = int(self.Header['Len'].cast(gdb.lookup_type('uint16')))
+		if IsWideString == True:
+			WideString = self.WideName.cast(gdb.lookup_type('WIDECHAR').pointer())
+			return '%s' % WideString.string('','',Len)
+		else:
+			AnsiString = self.AnsiName.cast(gdb.lookup_type('ANSICHAR').pointer())
+			return '%s' % AnsiString.string('','',Len)
+
+	def display_hint (self):
+		return 'string'
+
+# ------------------------------------------------------------------------------
+# FName
+#
+class FNamePrinter:
+	"Print FName"
+
+	def __init__(self, typename, val):
+		self.Value = val
+
+	def to_string(self):
+		if self.Value.is_optimized_out:
+			return '<optimized out>'
+
+		# ComparisonIndex is an FNameEntryId
+		Index = self.Value['ComparisonIndex']['Value']
+		IndexValue = int(Index)
+
+		if IndexValue >= 4194304:
+			return 'invalid'
+		else:
+			Expr = '((FNameEntry&)GNameBlocksDebug['+str(IndexValue)+' >> FNameDebugVisualizer::OffsetBits][FNameDebugVisualizer::EntryStride * ('+str(IndexValue)+' & FNameDebugVisualizer::OffsetMask)])'
+			NameEntry = gdb.parse_and_eval(Expr)
+			Number = self.Value['Number']
+			NumberValue = int(Number)
+			if NumberValue == 0:
+				return NameEntry
+			else:
+				return str([str(NameEntry), 'Number=' + str(NumberValue)])
+
+
+# ------------------------------------------------------------------------------
+# FMinimalName
+#
+class FMinimalNamePrinter:
+	"Print FMinimalName"
+
+	def __init__(self, typename, val):
+		self.Value = val
+
+	def to_string(self):
+		Index = self.Value['Index']['Value']
+		IndexValue = int(Index)
+
+		if IndexValue >= 4194304:
+			return 'invalid'
+		else:
+			Expr = '((FNameEntry&)GNameBlocksDebug['+str(IndexValue)+' >> FNameDebugVisualizer::OffsetBits][FNameDebugVisualizer::EntryStride * ('+str(IndexValue)+' & FNameDebugVisualizer::OffsetMask)])'
+			NameEntry = gdb.parse_and_eval(Expr)
+			Number = self.Value['Number']
+			NumberValue = int(Number)
+			if NumberValue == 0:
+				return NameEntry
+			else:
+				return str([str(NameEntry), 'Number=' + str(NumberValue)])
+
+# ------------------------------------------------------------------------------
+# TTuple
+#
+class TTuplePrinter:
+	"Print TTuple"
+
+	def __init__(self, typename, val):
+		self.Value = val
+
+		try:
+			self.TKey = self.Value["Key"];
+			self.TValue = self.Value["Value"];
+		except:
+			pass
+
+	def to_string(self):
+		return '(%s, %s)' % (self.TKey, self.TValue)
+
+
+# ------------------------------------------------------------------------------
+# TArray
+#
+class TArrayPrinter:
+	"Print TArray"
+
+	class _iterator(Iterator):
+		def __init__(self, val):
+			self.Value = val
+			self.Counter = -1
+			self.TType = self.Value.type.template_argument(0)
+
+			try:
+				self.ArrayNum = self.Value['ArrayNum']
+				if self.ArrayNum.is_optimized_out:
+					self.ArrayNum = 0
+
+				if self.ArrayNum > 0:
+					self.AllocatorInstance = self.Value['AllocatorInstance']
+					self.AllocatorInstanceData = self.AllocatorInstance['Data']
+					try:
+						self.InlineData = self.AllocatorInstance['InlineData']
+						self.SecondaryData = self.AllocatorInstance['SecondaryData']
+						if self.SecondaryData != None:
+							self.SecondaryDataData = self.SecondaryData['Data']
+						else:
+							self.SecondaryDataData = None
+					except:
+						pass
+			except:
+				raise
+
+		def __iter__(self):
+			return self
+
+		def __next__(self):
+			if self.ArrayNum == 0:
+				raise StopIteration
+
+			self.Counter = self.Counter + 1
+
+			if self.Counter >= self.ArrayNum:
+				raise StopIteration
+
+			try:
+				if self.AllocatorInstanceData != None:
+					data = self.AllocatorInstanceData.cast(self.TType.pointer())
+				elif self.SecondaryDataDataVal > 0:
+					data = self.SecondaryDataData.cast(self.TType.pointer())
+				else:
+					data = self.InlineData.cast(self.TType.pointer())
+			except:
+				return ('[%d]' % self.Counter, "optmized")
+
+			return ('[%d]' % self.Counter, data[self.Counter])
+
+	def __init__(self, typename, val):
+		self.Value = val;
+		self.ArrayNum = self.Value['ArrayNum']
+
+	def to_string(self):
+		if self.ArrayNum.is_optimized_out:
+			pass
+		if self.ArrayNum == 0:
+			return 'empty'
+		pass
+
+	def children(self):
+		return self._iterator(self.Value)
+
+	def display_hint(self):
+		return 'array'
+
+#
+# Register our lookup function. If no objfile is passed use all globally.
+def register_ue_printers(objfile):
+	if objfile == None:
+		objfile = gdb
+
+	objfile.pretty_printers.append(lookup_function)
+
+
+#
+# We need this part which is a definition how pretty printers work.
+#
+def lookup_function (val):
+	"Look-up and return a pretty-printer that can print val."
+
+	# Get the type and check if it points to a reference. We check for both Object and Pointer type.
+	type = val.type;
+	if (type.code == gdb.TYPE_CODE_REF) or (type.code == gdb.TYPE_CODE_PTR):
+		type = type.target ()
+
+	# Get the unqualified type, mean remove const nor volatile and strippe of typedefs.
+	type = type.unqualified ().strip_typedefs ()
+
+	# Get the tag name. The tag name is the name after struct, union, or enum in C and C++
+	tag = type.tag
+	if tag == None:
+		return None
+
+	for function in pretty_printers_dict:
+		if function.search (tag):
+			return pretty_printers_dict[function] (tag, val)
+
+	# Cannot find a pretty printer.  Return None.
+	return None
+
+def build_dictionary ():
+	pretty_printers_dict[re.compile('^FString$')] = lambda typename, val: FStringPrinter(typename, val)
+	pretty_printers_dict[re.compile('^FNameEntry$')] = lambda typename, val: FNameEntryPrinter(typename, val)
+	pretty_printers_dict[re.compile('^FName$')] = lambda typename, val: FNamePrinter(typename, val)
+	pretty_printers_dict[re.compile('^FMinimalName$')] = lambda typename, val: FMinimalNamePrinter(typename, val)
+	pretty_printers_dict[re.compile('^TArray<.+,.+>$')] = lambda typename, val: TArrayPrinter(typename, val)
+	pretty_printers_dict[re.compile('^TBitArray<.+>$')] = lambda typename, val: TBitArrayPrinter(typename, val)
+	pretty_printers_dict[re.compile('^TChunkedArray<.+>$')] = lambda typename, val: TChunkedArrayPrinter(typename, val)
+	pretty_printers_dict[re.compile('^TSparseArray<.+>$')] = lambda typename, val: TSparseArrayPrinter(typename, val)
+	pretty_printers_dict[re.compile('^TSetElement<.+>$')] = lambda typename, val: TSetElementPrinter(typename, val)
+	pretty_printers_dict[re.compile('^TSet<.+>$')] = lambda typename, val: TSetPrinter(typename, val)
+	pretty_printers_dict[re.compile('^FBitReference$')] = lambda typename, val: FBitReferencePrinter(typename, val)
+	pretty_printers_dict[re.compile('^TMap<.+,.+,.+>$')] = lambda typename, val: TMapPrinter(typename, val)
+	pretty_printers_dict[re.compile('^TPair<.+,.+>$')] = lambda typename, val: TTuplePrinter(typename, val)
+	pretty_printers_dict[re.compile('^TTuple<.+,.+>$')] = lambda typename, val: TTuplePrinter(typename, val)
+	pretty_printers_dict[re.compile('^TWeakObjectPtr<.+>$')] = lambda typename, val: TWeakObjectPtrPrinter(typename, val)
+
+
+pretty_printers_dict = {}
+build_dictionary ()
diff --git a/Engine/Extras/LLDBDataFormatters/UE4DataFormatters.py b/Engine/Extras/LLDBDataFormatters/UE4DataFormatters.py
deleted file mode 100644
diff --git a/Engine/Extras/LLDBDataFormatters/UE4DataFormatters_2ByteChars.py b/Engine/Extras/LLDBDataFormatters/UE4DataFormatters_2ByteChars.py
deleted file mode 100644
diff --git a/Engine/Extras/LLDBDataFormatters/UEDataFormatters.py b/Engine/Extras/LLDBDataFormatters/UEDataFormatters.py
new file mode 100644
--- /dev/null
+++ b/Engine/Extras/LLDBDataFormatters/UEDataFormatters.py
@@ -0,0 +1,688 @@
+# Copyright Epic Games, Inc. All Rights Reserved.
+
+# /*=============================================================================
+#	LLDB Data Formatters for Unreal Types
+# =============================================================================*/
+
+import lldb
+import lldb.formatters.Logger
+
+# Uncomment the line below to have the data formatters emit debug logging
+# lldb.formatters.Logger._lldb_formatters_debug_level=1
+
+# What documentation there is for parsing values in LLDB can be found here:
+# https://lldb.llvm.org/python_reference/index.html
+# https://lldb.llvm.org/python_reference/lldb.SBValue-class.html
+
+# To install:
+# 1) Open Terminal and run:
+#        touch ~/.lldbinit
+#        open ~/.lldbinit
+# 2) Add the following text to .lldbini and save - modifying the path as appropriate:
+#        settings set target.inline-breakpoint-strategy always
+#        command script import "/Path/To/Epic/UE/Engine/Extras/LLDBDataFormatters/UEDataFormatters.py"
+# 3) Restart Xcode
+
+def UETCharSummaryProvider(valobj,dict):
+    Data = valobj.GetValue()
+    Val = valobj.GetSummary()
+    Type = valobj.GetType().GetUnqualifiedType()
+    if Type.IsPointerType():
+        DataVal = valobj.GetValueAsUnsigned(0)
+        if DataVal == 0:
+            Val = 'NULL'
+        else:
+            Expr = '(wchar_t*)('+str(Data)+')'
+            ValRef = valobj.CreateValueFromExpression('string', Expr)
+            Val = str(ValRef.GetSummary())
+    elif Type.IsReferenceType():
+        Expr = '(wchar_t&)('+str(Data)+')'
+        ValRef = valobj.CreateValueFromExpression('string', Expr)
+        Val = str(ValRef.GetSummary())
+    return Val
+
+def UEFStringSummaryProvider(valobj,dict):
+    Data = valobj.GetChildMemberWithName('Data')
+    ArrayNumVal = Data.GetNumChildren()
+    if ArrayNumVal < 0:
+       return 'string=Invalid'
+    elif ArrayNumVal == 0:
+       return 'string=Empty'
+    else:
+        AllocatorInstance = Data.GetChildMemberWithName('AllocatorInstance')
+        ActualData = AllocatorInstance.GetChildMemberWithName('Data')
+        Expr = '(wchar_t*)('+str(ActualData.GetValue())+')'
+        ValRef = valobj.CreateValueFromExpression('string', Expr)
+        Val = str(ValRef.GetSummary())
+        return 'string=' + Val
+
+def UEFNameEntrySummaryProvider(valobj,dict):
+    Header = valobj.GetChildMemberWithName('Header')
+    DataPtr = valobj.GetChildAtIndex(1).AddressOf().GetValueAsUnsigned(0)
+    IsWide = Header.GetChildMemberWithName('bIsWide').GetValueAsUnsigned(0)
+    Len = min(Header.GetChildMemberWithName('Len').GetValueAsUnsigned(0), 1023)
+    if IsWide:
+        SizeOfTChar = valobj.CreateValueFromExpression('size', 'sizeof(TCHAR))').GetValueAsUnsigned(0)
+        Encoding = "utf-16" if SizeOfTChar == 2 else "utf-32"
+        NumBytes = Len * SizeOfTChar
+        Data = valobj.process.ReadMemory(DataPtr,NumBytes,lldb.SBError())
+        Name = Data.decode(Encoding).encode("utf-8")
+    else:
+        Name = valobj.process.ReadMemory(DataPtr,Len,lldb.SBError())
+    return 'name=%s' % Name
+
+def UEFNameSummaryProvider(valobj,dict):
+    EntryId = valobj.GetChildMemberWithName('DisplayIndex')
+    if not EntryId.IsValid():
+        EntryId = valobj.GetChildMemberWithName('ComparisonIndex')
+    if not EntryId.IsValid():
+        EntryId = valobj.GetChildMemberWithName('Index')
+    Index = EntryId.GetChildMemberWithName('Value').GetValueAsUnsigned(0)
+    Number = valobj.GetChildMemberWithName('Number').GetValueAsUnsigned(0)
+
+    # FNameDebugVisualizer::OffsetBits = 16
+    # FNameDebugVisualizer::OffsetMask = (1 << OffsetBits) - 1 = 65535
+    NameEntryExpr = '(FNameEntry*)(GNameBlocksDebug['+str(Index)+' >> 16]+((alignof(FNameEntry) * ('+str(Index)+' & 65535))))'
+    NameEntry = valobj.CreateValueFromExpression('NameEntry', NameEntryExpr)
+    NameStr = UEFNameEntrySummaryProvider(NameEntry, dict)
+    if Number != 0:
+        return '%s_%s' % (NameStr, Number-1)
+    else:
+        return '%s' % NameStr
+
+def UEUObjectBaseSummaryProvider(valobj,dict):
+    Name = valobj.GetChildMemberWithName('NamePrivate')
+    return Name.GetSummary()
+
+def UEFFieldClassSummaryProvider(valobj,dict):
+    Name = valobj.GetChildMemberWithName('Name')
+    return Name.GetSummary()
+
+def UEFFieldSummaryProvider(valobj,dict):
+    Name = valobj.GetChildMemberWithName('NamePrivate')
+    return Name.GetSummary()
+
+class UETWeakObjectPtrSynthProvider:
+
+    def __init__(self, valobj, dict):
+        logger = lldb.formatters.Logger.Logger()
+        self.valobj = valobj
+
+    def num_children(self):
+        logger = lldb.formatters.Logger.Logger()
+        return 1
+
+    def get_child_index(self,name):
+        logger = lldb.formatters.Logger.Logger()
+        return 0
+
+    def get_child_at_index(self,index):
+        logger = lldb.formatters.Logger.Logger()
+        logger >> "Retrieving child " + str(index)
+        if self.ObjectSerialNumberVal >= 1:
+            Expr = 'GObjectArrayForDebugVisualizers->Objects['+str(self.ObjectIndexVal/65536)+']['+str(self.ObjectIndexVal%65536)+'].SerialNumber == '+str(self.ObjectSerialNumberVal)
+            Val = self.valobj.CreateValueFromExpression(str(self.ObjectIndexVal), Expr)
+            Value = Val.GetValueAsUnsigned(0)
+            if Value != 0:
+                Expr = 'GObjectArrayForDebugVisualizers->Objects['+str(self.ObjectIndexVal/65536)+']['+str(self.ObjectIndexVal%65536)+'].Object'
+                return self.valobj.CreateValueFromExpression('Object', Expr)
+            else:
+                Expr = '(void*)0xDEADBEEF'
+                return self.valobj.CreateValueFromExpression('Object', Expr)
+
+        Expr = 'nullptr'
+        return self.valobj.CreateValueFromExpression('Object', Expr)
+
+    def update(self):
+        logger = lldb.formatters.Logger.Logger()
+        try:
+            self.ObjectSerialNumber = self.valobj.GetChildMemberWithName('ObjectSerialNumber')
+            self.ObjectSerialNumberVal = self.ObjectSerialNumber.GetValueAsSigned(0)
+            self.ObjectIndex = self.valobj.GetChildMemberWithName('ObjectIndex')
+            self.ObjectIndexVal = self.ObjectIndex.GetValueAsSigned(0)
+        except:
+            pass
+
+    def has_children(self):
+            return True
+
+def UEFWeakObjectPtrSummaryProvider(valobj,dict):
+    ObjectSerialNumber = valobj.GetChildMemberWithName('ObjectSerialNumber')
+    ObjectSerialNumberVal = ObjectSerialNumber.GetValueAsSigned(0)
+    if ObjectSerialNumberVal < 1:
+        return 'object=nullptr'
+    ObjectIndex = valobj.GetChildMemberWithName('ObjectIndex')
+    ObjectIndexVal = ObjectIndex.GetValueAsSigned(0)
+    Expr = 'GObjectArrayForDebugVisualizers->Objects['+str(ObjectIndexVal/65536)+']['+str(ObjectIndexVal%65536)+'].SerialNumber == '+str(ObjectSerialNumberVal)
+    Val = valobj.CreateValueFromExpression(str(ObjectIndexVal), Expr)
+    ValRef = Val.GetValueAsUnsigned(0)
+    if ValRef == 0:
+        return 'object=STALE'
+    else:
+        Expr = 'GObjectArrayForDebugVisualizers->Objects['+str(ObjectIndexVal/65536)+']['+str(ObjectIndexVal%65536)+'].Object'
+        Val = valobj.CreateValueFromExpression(str(ObjectIndexVal), Expr)
+        return 'object=' + Val.GetValue()
+
+class UEChunkedArraySynthProvider:
+
+    def __init__(self, valobj, dict):
+        logger = lldb.formatters.Logger.Logger()
+        self.valobj = valobj
+
+    def num_children(self):
+        logger = lldb.formatters.Logger.Logger()
+        try:
+            NumElementsVal = self.NumElements.GetValueAsSigned(0)
+            return NumElementsVal;
+        except:
+            return 0;
+
+    def get_child_index(self,name):
+        logger = lldb.formatters.Logger.Logger()
+        try:
+            return int(name.lstrip('[').rstrip(']'))
+        except:
+            return None
+
+    def get_child_at_index(self,index):
+        logger = lldb.formatters.Logger.Logger()
+        logger >> "Retrieving child " + str(index)
+        Expr = '(unsigned)sizeof('+str(self.valobj.GetType().GetUnqualifiedType().GetName())+'::FChunk)/'+str(self.ElementTypeSize)
+        self.ChunkBytes = self.valobj.CreateValueFromExpression('['+str(index)+']', Expr)
+        self.ChunkBytesSize = self.ChunkBytes.GetValueAsUnsigned(0)
+        assert self.ChunkBytesSize != 0
+
+        Expr = '*(*((('+str(self.ElementType.GetName())+'**)'+str(self.AllocatorData.GetValue())+')+'+str(index / self.ChunkBytesSize)+')+'+str(index % self.ChunkBytesSize)+')'
+        return self.valobj.CreateValueFromExpression('['+str(index)+']', Expr)
+
+    def extract_type(self):
+        logger = lldb.formatters.Logger.Logger()
+        ArrayType = self.valobj.GetType().GetUnqualifiedType()
+        if ArrayType.IsReferenceType():
+            ArrayType = ArrayType.GetDereferencedType()
+        elif ArrayType.IsPointerType():
+            ArrayType = ArrayType.GetPointeeType()
+        if ArrayType.GetNumberOfTemplateArguments() > 0:
+            ElementType = ArrayType.GetTemplateArgumentType(0)
+        else:
+            ElementType = None
+        return ElementType
+
+    def update(self):
+        logger = lldb.formatters.Logger.Logger()
+        try:
+            self.ElementType = self.extract_type()
+            self.ElementTypeSize = self.ElementType.GetByteSize()
+            assert self.ElementTypeSize != 0
+            self.NumElements = self.valobj.GetChildMemberWithName('NumElements')
+            self.Chunks = self.valobj.GetChildMemberWithName('Chunks')
+            self.ArrayNum = self.Chunks.GetChildMemberWithName('ArrayNum')
+            self.AllocatorInstance = self.Chunks.GetChildMemberWithName('AllocatorInstance')
+            self.AllocatorData = self.AllocatorInstance.GetChildMemberWithName('Data')
+        except:
+            pass
+
+    def has_children(self):
+        return True
+
+def UEChunkedArraySummaryProvider(valobj,dict):
+    return 'size=' + str(valobj.GetNumChildren())
+
+class UESparseArraySynthProvider:
+
+    def __init__(self, valobj, dict):
+        logger = lldb.formatters.Logger.Logger()
+        self.valobj = valobj
+
+    def num_children(self):
+        logger = lldb.formatters.Logger.Logger()
+        try:
+            NumBitsVal = self.NumFreeIndices.GetValueAsSigned(0)
+            ArrayNumVal = self.Data.GetChildMemberWithName('ArrayNum').GetValueAsSigned(0)
+            return ArrayNumVal - NumBitsVal;
+        except:
+            return 0;
+
+    def get_child_index(self,name):
+        logger = lldb.formatters.Logger.Logger()
+        try:
+            return int(name.lstrip('[').rstrip(']'))
+        except:
+            return None
+
+    def get_child_at_index(self,index):
+        logger = lldb.formatters.Logger.Logger()
+        logger >> "Retrieving child " + str(index)
+        if index < 0:
+            return None;
+
+        if index >= self.num_children():
+            return None;
+
+        Val = None
+        if self.SecondaryDataDataVal > 0:
+            Expr = '(bool)((((int*)'+str(self.SecondaryDataData.GetAddress())+')['+str(index)+'/32] >> '+str(index)+') & 1)'
+            Val = self.SecondaryDataData.CreateValueFromExpression('['+str(index)+']', Expr)
+        else:
+            Expr = '(bool)((((int*)'+str(self.InlineData.GetAddress())+')['+str(index)+'/32] >> '+str(index)+') & 1)'
+            Val = self.InlineData.CreateValueFromExpression('['+str(index)+']', Expr)
+
+        if Val.GetValueAsSigned(0) != 0:
+            offset = index * self.ElementTypeSize
+            return self.AllocatorData.CreateChildAtOffset('['+str(index)+']',offset,self.ElementType)
+        else:
+            return None
+
+    def extract_type(self):
+        logger = lldb.formatters.Logger.Logger()
+        ArrayType = self.valobj.GetType().GetUnqualifiedType()
+        if ArrayType.IsReferenceType():
+            ArrayType = ArrayType.GetDereferencedType()
+        elif ArrayType.IsPointerType():
+            ArrayType = ArrayType.GetPointeeType()
+        if ArrayType.GetNumberOfTemplateArguments() > 0:
+            ElementType = ArrayType.GetTemplateArgumentType(0)
+        else:
+            ElementType = None
+        return ElementType
+
+    def update(self):
+        logger = lldb.formatters.Logger.Logger()
+        try:
+            self.ElementType = self.extract_type()
+            self.ElementTypeSize = self.ElementType.GetByteSize()
+            assert self.ElementTypeSize != 0
+            self.NumFreeIndices = self.valobj.GetChildMemberWithName('NumFreeIndices')
+            self.Data = self.valobj.GetChildMemberWithName('Data')
+            self.AllocatorInstance = self.Data.GetChildMemberWithName('AllocatorInstance')
+            self.AllocatorData = self.AllocatorInstance.GetChildMemberWithName('Data')
+            self.AllocationFlags = self.valobj.GetChildMemberWithName('AllocationFlags')
+            self.InlineData = self.AllocationFlags.GetChildMemberWithName('InlineData')
+            self.SecondaryData = self.AllocationFlags.GetChildMemberWithName('SecondaryData')
+            self.SecondaryDataData = self.SecondaryData.GetChildMemberWithName('Data')
+            self.SecondaryDataDataVal = self.SecondaryDataData.GetValueAsSigned(0)
+        except:
+            pass
+
+    def has_children(self):
+        return True
+
+def UESparseArraySummaryProvider(valobj,dict):
+    return 'size=' + str(valobj.GetNumChildren())
+
+class UEBitArraySynthProvider:
+
+    def __init__(self, valobj, dict):
+        logger = lldb.formatters.Logger.Logger()
+        self.valobj = valobj
+
+    def num_children(self):
+        logger = lldb.formatters.Logger.Logger()
+        try:
+            NumBitsVal = self.NumBits.GetValueAsSigned(0)
+            return NumBitsVal;
+        except:
+            return 0;
+
+    def get_child_index(self,name):
+        logger = lldb.formatters.Logger.Logger()
+        try:
+            return int(name.lstrip('[').rstrip(']'))
+        except:
+            return None
+
+    def get_child_at_index(self,index):
+        logger = lldb.formatters.Logger.Logger()
+        logger >> "Retrieving child " + str(index)
+        if index < 0:
+            return None;
+
+        if index >= self.num_children():
+            return None;
+
+        if self.SecondaryDataDataVal > 0:
+            Expr = '(bool)((((int*)'+str(self.SecondaryDataData.GetAddress())+')['+str(index)+'/32] >> '+str(index)+') & 1)'
+            return self.SecondaryDataData.CreateValueFromExpression('['+str(index)+']', Expr)
+        else:
+            Expr = '(bool)((((int*)'+str(self.InlineData.GetAddress())+')['+str(index)+'/32] >> '+str(index)+') & 1)'
+            return self.InlineData.CreateValueFromExpression('['+str(index)+']', Expr)
+
+        return None
+
+    def update(self):
+        logger = lldb.formatters.Logger.Logger()
+        try:
+            self.SecondaryDataData = None
+            self.SecondaryDataDataVal = 0
+            self.NumBits = self.valobj.GetChildMemberWithName('NumBits')
+            self.MaxBits = self.valobj.GetChildMemberWithName('MaxBits')
+            self.InlineData = self.valobj.GetChildMemberWithName('InlineData')
+            self.SecondaryData = self.valobj.GetChildMemberWithName('SecondaryData')
+            if self.SecondaryData != None:
+                self.SecondaryDataData = self.SecondaryData.GetChildMemberWithName('Data')
+                self.SecondaryDataDataVal = self.SecondaryDataData.GetValueAsUnsigned(0)
+        except:
+            pass
+
+    def has_children(self):
+        return True
+
+def UEBitArraySummaryProvider(valobj,dict):
+    return 'size=' + str(valobj.GetNumChildren())
+
+class UEArraySynthProvider:
+
+    def __init__(self, valobj, dict):
+        logger = lldb.formatters.Logger.Logger()
+        self.valobj = valobj
+
+    def num_children(self):
+        logger = lldb.formatters.Logger.Logger()
+        try:
+            ArrayNumVal = self.ArrayNum.GetValueAsSigned(0)
+            return ArrayNumVal + self.NumChildren;
+        except:
+            return 0;
+
+    def get_child_index(self,name):
+        logger = lldb.formatters.Logger.Logger()
+        try:
+            return self.NumChildren + int(name.lstrip('[').rstrip(']'))
+        except:
+            return self.valobj.GetIndexOfChildWithName(name)
+
+    def get_child_at_index(self,index):
+        logger = lldb.formatters.Logger.Logger()
+        logger >> "Retrieving child " + str(index)
+        if index < 0:
+            return None;
+
+        if index < self.NumChildren:
+            logger >> "Existing child " + str(index)
+            return self.valobj.GetChildAtIndex(index)
+        else:
+            index -= self.NumChildren
+
+        if index >= self.num_children():
+            logger >> "Index is greater than num children " + str(self.num_children())
+            return None;
+        try:
+            offset = index * self.ElementTypeSize
+            if self.Data != None and self.Data.IsValid():
+                logger >> "Retrieving Data child of type " + str(self.ElementType) + " at offset " + str(offset)
+                return self.Data.CreateChildAtOffset('['+str(index)+']',offset,self.ElementType)
+            elif self.SecondaryDataDataVal > 0:
+                logger >> "Retrieving SecondaryData child of type " + str(self.ElementType) + " at offset " + str(offset)
+                return self.SecondaryDataData.CreateChildAtOffset('['+str(index)+']',offset,self.ElementType)
+            else:
+                logger >> "Retrieving InlineData child of type " + str(self.ElementType) + " at offset " + str(offset)
+                return self.InlineData.CreateChildAtOffset('['+str(index)+']',offset,self.ElementType)
+        except:
+            logger >> "Exception retrieving child of type " + str(self.ElementType) + " at index " + str(index)
+            return None
+
+    def extract_type(self):
+        logger = lldb.formatters.Logger.Logger()
+        ArrayType = self.valobj.GetType().GetUnqualifiedType()
+        if ArrayType.IsReferenceType():
+            ArrayType = ArrayType.GetDereferencedType()
+        elif ArrayType.IsPointerType():
+            ArrayType = ArrayType.GetPointeeType()
+
+        if ArrayType.GetNumberOfTemplateArguments() > 0:
+            ElementType = ArrayType.GetTemplateArgumentType(0)
+        else:
+            ElementType = None
+        return ElementType
+
+    def update(self):
+        logger = lldb.formatters.Logger.Logger()
+        try:
+            self.NumChildren = self.valobj.GetNumChildren()
+            self.ArrayNum = self.valobj.GetChildMemberWithName('ArrayNum')
+            self.ArrayMax = self.valobj.GetChildMemberWithName('ArrayMax')
+            self.AllocatorInstance = self.valobj.GetChildMemberWithName('AllocatorInstance')
+            if self.AllocatorInstance.GetType().IsReferenceType():
+                self.AllocatorInstance = self.AllocatorInstance.Dereference()
+            self.Data = None
+            self.Data = self.AllocatorInstance.GetChildMemberWithName('Data')
+            self.InlineData = self.AllocatorInstance.GetChildMemberWithName('InlineData')
+            self.SecondaryData = self.AllocatorInstance.GetChildMemberWithName('SecondaryData')
+            if self.SecondaryData != None:
+                self.SecondaryDataData = self.SecondaryData.GetChildMemberWithName('Data')
+                self.SecondaryDataDataVal = self.SecondaryDataData.GetValueAsUnsigned(0)
+            else:
+                self.SecondaryDataData = None
+                self.SecondaryDataDataVal = 0
+        except:
+            logger >> "UEArraySynthProvider::update failed accessing members"
+            pass
+        try:
+            self.ElementType = self.extract_type()
+            self.ElementTypeSize = self.ElementType.GetByteSize()
+            assert self.ElementTypeSize != 0
+        except:
+            logger >> "UEArraySynthProvider::update failed accessing element type"
+            pass
+
+    def has_children(self):
+        return True
+
+def UEArraySummaryProvider(valobj,dict):
+    return 'size=' + str(valobj.GetChildMemberWithName('ArrayNum').GetValueAsSigned(0))
+
+class UESetSynthProvider:
+
+    def __init__(self, valobj, dict):
+        logger = lldb.formatters.Logger.Logger()
+        self.valobj = valobj
+
+    def num_children(self):
+        logger = lldb.formatters.Logger.Logger()
+        try:
+            self.ArrayNumVal = self.ArrayNum.GetValueAsUnsigned(0)
+            self.NumFreeIndicesVal = self.NumFreeIndices.GetValueAsUnsigned(0)
+            return self.ArrayNumVal - self.NumFreeIndicesVal;
+        except:
+            return 0;
+
+    def get_child_index(self,name):
+        logger = lldb.formatters.Logger.Logger()
+        try:
+            return int(name.lstrip('[').rstrip(']'))
+        except:
+            return 0
+
+    def get_child_at_index(self,index):
+        logger = lldb.formatters.Logger.Logger()
+        logger >> "Retrieving child " + str(index)
+        if index < 0:
+            return None;
+
+        if index >= self.num_children():
+            return None;
+        try:
+            offset = index * self.ElementTypeSize
+            HasObject = 0
+            if self.SecondaryDataDataVal > 0:
+                Expr = '(bool)((((int*)'+str(self.SecondaryDataDataVal)+')['+str(index)+'/32] >> '+str(index)+') & 1)'
+                HasObject = 1 ##self.AllocationFlagsSecondaryDataData.CreateValueFromExpression('['+str(index)+']', Expr).GetValueAsUnsigned(0)
+            else:
+                Expr = '(bool)((((int*)'+str(self.AllocationFlagsInlineDataAddr)+')['+str(index)+'/32] >> '+str(index)+') & 1)'
+                HasObject = 1 ##self.AllocationFlagsInlineData.CreateValueFromExpression('['+str(index)+']', Expr).GetValueAsUnsigned(0)
+
+            if HasObject == 1:
+                return self.AllocatorInstanceData.CreateChildAtOffset('['+str(index)+']',offset,self.ElementType)
+            else:
+                return self.valobj.CreateValueFromExpression('['+str(index)+']', '(void*)0xDEADBEEF')
+        except:
+            return None
+
+    def extract_type(self):
+        logger = lldb.formatters.Logger.Logger()
+        ArrayType = self.valobj.GetType().GetUnqualifiedType()
+        if ArrayType.IsReferenceType():
+            ArrayType = ArrayType.GetDereferencedType()
+        elif ArrayType.IsPointerType():
+            ArrayType = ArrayType.GetPointeeType()
+        if ArrayType.GetNumberOfTemplateArguments() > 0:
+            ElementType = ArrayType.GetTemplateArgumentType(0)
+        else:
+            ElementType = None
+        return ElementType
+
+    def update(self):
+        logger = lldb.formatters.Logger.Logger()
+        try:
+            self.Elements = self.valobj.GetChildMemberWithName('Elements')
+            self.ElementsData = self.Elements.GetChildMemberWithName('Data')
+            self.ArrayNum = self.ElementsData.GetChildMemberWithName('ArrayNum')
+            self.NumFreeIndices = self.Elements.GetChildMemberWithName('NumFreeIndices')
+            self.AllocationFlags = self.Elements.GetChildMemberWithName('AllocationFlags')
+            self.AllocationFlagsAllocatorInstance = self.AllocationFlags.GetChildMemberWithName('AllocatorInstance')
+            self.AllocatorInstance = self.ElementsData.GetChildMemberWithName('AllocatorInstance')
+            self.AllocatorInstanceData = self.AllocatorInstance.GetChildMemberWithName('Data')
+            self.AllocationFlagsInlineData = self.AllocationFlagsAllocatorInstance.GetChildMemberWithName('InlineData')
+            self.AllocationFlagsInlineDataAddr = self.AllocationFlagsInlineData.AddressOf().GetValueAsUnsigned(0)
+            self.AllocationFlagsSecondaryData = self.AllocationFlagsAllocatorInstance.GetChildMemberWithName('SecondaryData')
+            self.AllocationFlagsSecondaryDataData = self.AllocationFlagsSecondaryData.GetChildMemberWithName('Data')
+            self.SecondaryDataDataVal = self.AllocationFlagsSecondaryDataData.GetValueAsUnsigned(0)
+            self.ElementType = self.extract_type()
+            # Can't cast to TSetElement - the template instantiation won't allow it
+            Expr = '(size_t)sizeof(FSetElementId) + sizeof(int32)'
+            TSetElement = self.valobj.CreateValueFromExpression('TSetElement', Expr)
+            # This may fail due to C++ struct padding - will have to check
+            self.ElementTypeSize = self.ElementType.GetByteSize() + TSetElement.GetValueAsUnsigned(0)
+            assert self.ElementTypeSize != 0
+        except:
+            pass
+
+    def has_children(self):
+        return True
+
+def UESetSummaryProvider(valobj,dict):
+    return 'size=' + str(valobj.GetNumChildren())
+
+class UEMapSynthProvider:
+
+    def __init__(self, valobj, dict):
+        logger = lldb.formatters.Logger.Logger()
+        self.valobj = valobj
+
+    def num_children(self):
+        logger = lldb.formatters.Logger.Logger()
+        try:
+            self.ArrayNumVal = self.ArrayNum.GetValueAsUnsigned(0)
+            self.NumFreeIndicesVal = self.NumFreeIndices.GetValueAsUnsigned(0)
+            return self.ArrayNumVal - self.NumFreeIndicesVal;
+        except:
+            return 0;
+
+    def get_child_index(self,name):
+        logger = lldb.formatters.Logger.Logger()
+        try:
+            return int(name.lstrip('[').rstrip(']'))
+        except:
+            return 0
+
+    def get_child_at_index(self,index):
+        logger = lldb.formatters.Logger.Logger()
+        logger >> "Retrieving child " + str(index)
+        if index < 0:
+            return None;
+
+        if index >= self.num_children():
+            return None;
+        try:
+            offset = index * self.ElementTypeSize
+            HasObject = 0
+            if self.SecondaryDataDataVal != 0:
+                Expr = '(bool)((((unsigned int*)'+str(self.SecondaryDataDataVal)+')['+str(index)+'/32] >> '+str(index)+') & 1)'
+                HasObject = 1 ##self.AllocationFlagsSecondaryDataData.CreateValueFromExpression('['+str(index)+']', Expr).GetValueAsUnsigned(0)
+            else:
+                Expr = '(bool)((((unsigned int*)'+str(self.AllocationFlagsInlineDataAddr)+')['+str(index)+'/32] >> '+str(index)+') & 1)'
+                HasObject = 1 ##self.AllocationFlagsInlineData.CreateValueFromExpression('['+str(index)+']', Expr).GetValueAsUnsigned(0)
+
+            if HasObject == 1:
+                return self.AllocatorInstanceData.CreateChildAtOffset('['+str(index)+']',offset,self.ElementType)
+            else:
+                return self.valobj.CreateValueFromExpression('['+str(index)+']', '(void*)0xDEADBEEF')
+        except:
+            return None
+
+    def extract_type(self):
+        logger = lldb.formatters.Logger.Logger()
+        ArrayType = self.Pairs.GetType().GetUnqualifiedType()
+        if ArrayType.IsReferenceType():
+            ArrayType = ArrayType.GetDereferencedType()
+        elif ArrayType.IsPointerType():
+            ArrayType = ArrayType.GetPointeeType()
+        if ArrayType.GetNumberOfTemplateArguments() > 0:
+            ElementType = ArrayType.GetTemplateArgumentType(0)
+        else:
+            ElementType = None
+        return ElementType
+
+    def update(self):
+        logger = lldb.formatters.Logger.Logger()
+        try:
+            self.Pairs = self.valobj.GetChildMemberWithName('Pairs')
+            self.Elements = self.Pairs.GetChildMemberWithName('Elements')
+            self.ElementsData = self.Elements.GetChildMemberWithName('Data')
+            self.ArrayNum = self.ElementsData.GetChildMemberWithName('ArrayNum')
+            self.NumFreeIndices = self.Elements.GetChildMemberWithName('NumFreeIndices')
+            self.AllocationFlags = self.Elements.GetChildMemberWithName('AllocationFlags')
+            self.AllocationFlagsAllocatorInstance = self.AllocationFlags.GetChildMemberWithName('AllocatorInstance')
+            self.AllocatorInstance = self.ElementsData.GetChildMemberWithName('AllocatorInstance')
+            self.AllocatorInstanceData = self.AllocatorInstance.GetChildMemberWithName('Data')
+            self.AllocationFlagsInlineData = self.AllocationFlagsAllocatorInstance.GetChildMemberWithName('InlineData')
+            self.AllocationFlagsInlineDataAddr = self.AllocationFlagsInlineData.AddressOf().GetValueAsUnsigned(0)
+            self.AllocationFlagsSecondaryData = self.AllocationFlagsAllocatorInstance.GetChildMemberWithName('SecondaryData')
+            self.AllocationFlagsSecondaryDataData = self.AllocationFlagsSecondaryData.GetChildMemberWithName('Data')
+            self.SecondaryDataDataVal = self.AllocationFlagsSecondaryDataData.GetValueAsUnsigned(0)
+            self.ElementType = self.extract_type()
+            # Can't cast to TSetElement - the template instantiation won't allow it
+            Expr = '(size_t)sizeof(FSetElementId) + sizeof(int32)'
+            TSetElement = self.valobj.CreateValueFromExpression('TSetElement', Expr)
+            # This may fail due to C++ struct padding - will have to check
+            self.ElementTypeSize = self.ElementType.GetByteSize() + TSetElement.GetValueAsUnsigned(0)
+            assert self.ElementTypeSize != 0
+        except:
+            pass
+
+    def has_children(self):
+        return True
+
+def UEMapSummaryProvider(valobj,dict):
+    return 'size=' + str(valobj.GetNumChildren())
+
+def __lldb_init_module(debugger,dict):
+    debugger.HandleCommand('type summary add -F UEDataFormatters.UETCharSummaryProvider -e TCHAR -w UEDataFormatters')
+    debugger.HandleCommand('type summary add -F UEDataFormatters.UEFStringSummaryProvider -e -x "FString$" -w UEDataFormatters')
+    debugger.HandleCommand('type summary add -F UEDataFormatters.UEFNameEntrySummaryProvider -e -x "FNameEntry$" -w UEDataFormatters')
+    debugger.HandleCommand('type summary add -F UEDataFormatters.UEFNameSummaryProvider -e -x "FName$" -w UEDataFormatters')
+    debugger.HandleCommand('type summary add -F UEDataFormatters.UEFNameSummaryProvider -e -x "FMinimalName$" -w UEDataFormatters')
+    debugger.HandleCommand('type summary add -F UEDataFormatters.UEUObjectBaseSummaryProvider -e UObject -w UEDataFormatters')
+    debugger.HandleCommand('type summary add -F UEDataFormatters.UEUObjectBaseSummaryProvider -e UObjectBase -w UEDataFormatters')
+    debugger.HandleCommand('type summary add -F UEDataFormatters.UEUObjectBaseSummaryProvider -e UObjectBaseUtility -w UEDataFormatters')
+    debugger.HandleCommand('type summary add -F UEDataFormatters.UEFFieldClassSummaryProvider -e FFieldClass -w UEDataFormatters')
+    debugger.HandleCommand('type summary add -F UEDataFormatters.UEFFieldSummaryProvider -e FField -w UEDataFormatters')
+    debugger.HandleCommand('type summary add -F UEDataFormatters.UEFWeakObjectPtrSummaryProvider -e FWeakObjectPtr -w UEDataFormatters')
+    debugger.HandleCommand('type synthetic add -l UEDataFormatters.UETWeakObjectPtrSynthProvider -x "TWeakObjectPtr<.+>$" -w UEDataFormatters')
+    debugger.HandleCommand('type synthetic add -l UEDataFormatters.UETWeakObjectPtrSynthProvider -x "TAutoWeakObjectPtr<.+>$" -w UEDataFormatters')
+    debugger.HandleCommand('type synthetic add -l UEDataFormatters.UEArraySynthProvider -x "TArray<.+,.+>$" -w UEDataFormatters')
+    debugger.HandleCommand('type summary add -F UEDataFormatters.UEArraySummaryProvider -e -x "TArray<.+>$" -w UEDataFormatters')
+    debugger.HandleCommand('type synthetic add -l UEDataFormatters.UEBitArraySynthProvider -x "TBitArray<.+>$" -w UEDataFormatters')
+    debugger.HandleCommand('type summary add -F UEDataFormatters.UEBitArraySummaryProvider -e -x "TBitArray<.+>$" -w UEDataFormatters')
+    debugger.HandleCommand('type synthetic add -l UEDataFormatters.UESparseArraySynthProvider -x "TSparseArray<.+>$" -w UEDataFormatters')
+    debugger.HandleCommand('type summary add -F UEDataFormatters.UESparseArraySummaryProvider -e -x "TSparseArray<.+>$" -w UEDataFormatters')
+    debugger.HandleCommand('type synthetic add -l UEDataFormatters.UEChunkedArraySynthProvider -x "TChunkedArray<.+>$" -w UEDataFormatters')
+    debugger.HandleCommand('type summary add -F UEDataFormatters.UEChunkedArraySummaryProvider -e -x "TChunkedArray<.+>$" -w UEDataFormatters')
+    debugger.HandleCommand('type synthetic add -l UEDataFormatters.UESetSynthProvider -x "TSet<.+>$" -w UEDataFormatters')
+    debugger.HandleCommand('type summary add -F UEDataFormatters.UESetSummaryProvider -e -x "TSet<.+>$" -w UEDataFormatters')
+    debugger.HandleCommand('type synthetic add -l UEDataFormatters.UEMapSynthProvider -x "TMap<.+>$" -w UEDataFormatters')
+    debugger.HandleCommand('type summary add -F UEDataFormatters.UEMapSummaryProvider -e -x "TMap<.+>$" -w UEDataFormatters')
+    debugger.HandleCommand('type synthetic add -l UEDataFormatters.UEMapSynthProvider -x "TMapBase<.+>$" -w UEDataFormatters')
+    debugger.HandleCommand('type summary add -F UEDataFormatters.UEMapSummaryProvider -e -x "TMapBase<.+>$" -w UEDataFormatters')
+    debugger.HandleCommand("type category enable UEDataFormatters")
+
diff --git a/Engine/Extras/LLDBDataFormatters/UEDataFormatters_2ByteChars.py b/Engine/Extras/LLDBDataFormatters/UEDataFormatters_2ByteChars.py
new file mode 100644
--- /dev/null
+++ b/Engine/Extras/LLDBDataFormatters/UEDataFormatters_2ByteChars.py
@@ -0,0 +1,683 @@
+# Copyright Epic Games, Inc. All Rights Reserved.
+
+# /*=============================================================================
+#	LLDB Data Formatters for Unreal Types
+# =============================================================================*/
+
+import lldb
+import lldb.formatters.Logger
+
+# Uncomment the line below to have the data formatters emit debug logging
+# lldb.formatters.Logger._lldb_formatters_debug_level=1
+
+# What documentation there is for parsing values in LLDB can be found here:
+# https://lldb.llvm.org/python_reference/index.html
+# https://lldb.llvm.org/python_reference/lldb.SBValue-class.html
+
+# To install:
+# 1) Open Terminal and run:
+#        touch ~/.lldbinit
+#        open ~/.lldbinit
+# 2) Add the following text to .lldbini and save - modifying the path as appropriate:
+#        settings set target.inline-breakpoint-strategy always
+#        command script import "/Path/To/Epic/UE/Engine/Extras/LLDBDataFormatters/UEDataFormatters_2ByteChars.py"
+# 3) Restart Xcode
+
+def UETCharSummaryProvider(valobj,dict):
+    Data = valobj.GetValue()
+    Val = valobj.GetSummary()
+    Type = valobj.GetType().GetUnqualifiedType()
+    if Type.IsPointerType():
+        DataVal = valobj.GetValueAsUnsigned(0)
+        if DataVal == 0:
+            Val = 'NULL'
+        else:
+            Expr = '(char16_t*)(%s)' % Data
+            ValRef = valobj.CreateValueFromExpression('string', Expr)
+            Val = ValRef.GetSummary()
+    elif Type.IsReferenceType():
+        Expr = '(char16_t&)(%s)' % Data
+        ValRef = valobj.CreateValueFromExpression('string', Expr)
+        Val = ValRef.GetSummary()
+    return Val
+
+def UEFStringSummaryProvider(valobj,dict):
+    Data = valobj.GetChildMemberWithName('Data')
+    ArrayNumVal = Data.GetChildMemberWithName('ArrayNum').GetValueAsSigned(0)
+    if ArrayNumVal < 0:
+       return 'string=Invalid'
+    elif ArrayNumVal == 0:
+       return 'string=Empty'
+    else:
+        AllocatorInstance = Data.GetChildMemberWithName('AllocatorInstance')
+        ActualData = AllocatorInstance.GetChildMemberWithName('Data')
+        Expr = '(char16_t*)(%s)' % ActualData.GetValue()
+        ValRef = valobj.CreateValueFromExpression('string', Expr)
+        Val = ValRef.GetSummary()
+        return 'string=' + Val
+
+def UEFNameEntrySummaryProvider(valobj,dict):
+    Header = valobj.GetChildMemberWithName('Header')
+    DataPtr = valobj.GetChildAtIndex(1).AddressOf().GetValueAsUnsigned(0)
+    IsWide = Header.GetChildMemberWithName('bIsWide').GetValueAsUnsigned(0)
+    Len = min(Header.GetChildMemberWithName('Len').GetValueAsUnsigned(0), 1023)
+    if IsWide:
+        SizeOfTChar = valobj.CreateValueFromExpression('size', 'sizeof(TCHAR))').GetValueAsUnsigned(0)
+        Encoding = "utf-16" if SizeOfTChar == 2 else "utf-32"
+        NumBytes = Len * SizeOfTChar
+        Data = valobj.process.ReadMemory(DataPtr,NumBytes,lldb.SBError())
+        Name = Data.decode(Encoding).encode("utf-8")
+    else:
+        Name = valobj.process.ReadMemory(DataPtr,Len,lldb.SBError())
+    return 'name=%s' % Name
+
+def UEFNameSummaryProvider(valobj,dict):
+    EntryId = valobj.GetChildMemberWithName('DisplayIndex')
+    if not EntryId.IsValid():
+        EntryId = valobj.GetChildMemberWithName('ComparisonIndex')
+    if not EntryId.IsValid():
+        EntryId = valobj.GetChildMemberWithName('Index')
+    Index = EntryId.GetChildMemberWithName('Value').GetValueAsUnsigned(0)
+    Number = valobj.GetChildMemberWithName('Number').GetValueAsUnsigned(0)
+
+    # FNameDebugVisualizer::OffsetBits = 16
+    # FNameDebugVisualizer::OffsetMask = (1 << OffsetBits) - 1 = 65535
+    NameEntryExpr = '(FNameEntry*)(GNameBlocksDebug[%s >> 16]+((alignof(FNameEntry) * (%s & 65535))))' % (Index, Index)
+    NameEntry = valobj.CreateValueFromExpression('NameEntry', NameEntryExpr)
+    NameStr = UEFNameEntrySummaryProvider(NameEntry, dict)
+    if Number != 0:
+        return '%s_%s' % (NameStr, Number-1)
+    else:
+        return '%s' % NameStr
+
+def UEUObjectBaseSummaryProvider(valobj,dict):
+    Name = valobj.GetChildMemberWithName('NamePrivate')
+    return Name.GetSummary()
+
+def UEFFieldClassSummaryProvider(valobj,dict):
+    Name = valobj.GetChildMemberWithName('Name')
+    return Name.GetSummary()
+
+def UEFFieldSummaryProvider(valobj,dict):
+    Name = valobj.GetChildMemberWithName('NamePrivate')
+    return Name.GetSummary()
+
+class UETWeakObjectPtrSynthProvider:
+
+    def __init__(self, valobj, dict):
+        logger = lldb.formatters.Logger.Logger()
+        self.valobj = valobj
+
+    def num_children(self):
+        logger = lldb.formatters.Logger.Logger()
+        return 1
+
+    def get_child_index(self,name):
+        logger = lldb.formatters.Logger.Logger()
+        return 0
+
+    def get_child_at_index(self,index):
+        logger = lldb.formatters.Logger.Logger()
+        logger >> "Retrieving child %s" % index
+        if self.ObjectSerialNumberVal >= 1:
+            Expr = 'GObjectArrayForDebugVisualizers->Objects[%s][%s].SerialNumber == %s' % (self.ObjectIndexVal/65536, self.ObjectIndexVal%65536, self.ObjectSerialNumberVal)
+            Val = self.valobj.CreateValueFromExpression("%s" % self.ObjectIndexVal, Expr)
+            Value = Val.GetValueAsUnsigned(0)
+            if Value != 0:
+                Expr = 'GObjectArrayForDebugVisualizers->Objects[%s][%s].Object' % (self.ObjectIndexVal/65536, self.ObjectIndexVal%65536)
+                return self.valobj.CreateValueFromExpression('Object', Expr)
+            else:
+                Expr = '(void*)0xDEADBEEF'
+                return self.valobj.CreateValueFromExpression('Object', Expr)
+
+        Expr = 'nullptr'
+        return self.valobj.CreateValueFromExpression('Object', Expr)
+
+    def update(self):
+        logger = lldb.formatters.Logger.Logger()
+        try:
+            self.ObjectSerialNumber = self.valobj.GetChildMemberWithName('ObjectSerialNumber')
+            self.ObjectSerialNumberVal = self.ObjectSerialNumber.GetValueAsSigned(0)
+            self.ObjectIndex = self.valobj.GetChildMemberWithName('ObjectIndex')
+            self.ObjectIndexVal = self.ObjectIndex.GetValueAsSigned(0)
+        except:
+            pass
+
+    def has_children(self):
+            return True
+
+def UEFWeakObjectPtrSummaryProvider(valobj,dict):
+    ObjectSerialNumber = valobj.GetChildMemberWithName('ObjectSerialNumber')
+    ObjectSerialNumberVal = ObjectSerialNumber.GetValueAsSigned(0)
+    if ObjectSerialNumberVal < 1:
+        return 'object=nullptr'
+    ObjectIndex = valobj.GetChildMemberWithName('ObjectIndex')
+    ObjectIndexVal = ObjectIndex.GetValueAsSigned(0)
+    Expr = 'GObjectArrayForDebugVisualizers->Objects[%s][%s].SerialNumber == %s' % (ObjectIndexVal/65536, ObjectIndexVal%65536, ObjectSerialNumberVal)
+    Val = valobj.CreateValueFromExpression('%s' % ObjectIndexVal, Expr)
+    ValRef = Val.GetValueAsUnsigned(0)
+    if ValRef == 0:
+        return 'object=STALE'
+    else:
+        Expr = 'GObjectArrayForDebugVisualizers->Objects[%s][%s].Object' % (ObjectIndexVal/65536, ObjectIndexVal%65536)
+        Val = valobj.CreateValueFromExpression("%s" % ObjectIndexVal, Expr)
+        return 'object=' + Val.GetValue()
+
+class UEChunkedArraySynthProvider:
+
+    def __init__(self, valobj, dict):
+        logger = lldb.formatters.Logger.Logger()
+        self.valobj = valobj
+
+    def num_children(self):
+        logger = lldb.formatters.Logger.Logger()
+        try:
+            NumElementsVal = self.NumElements.GetValueAsSigned(0)
+            return NumElementsVal;
+        except:
+            return 0;
+
+    def get_child_index(self,name):
+        logger = lldb.formatters.Logger.Logger()
+        try:
+            return int(name.lstrip('[').rstrip(']'))
+        except:
+            return None
+
+    def get_child_at_index(self,index):
+        logger = lldb.formatters.Logger.Logger()
+        logger >> "Retrieving child %s" % index
+        Expr = '(unsigned)sizeof(%s::FChunk)/%s' % (self.valobj.GetType().GetUnqualifiedType().GetName(), self.ElementTypeSize)
+        self.ChunkBytes = self.valobj.CreateValueFromExpression('[%s]' % index, Expr)
+        self.ChunkBytesSize = self.ChunkBytes.GetValueAsUnsigned(0)
+        assert self.ChunkBytesSize != 0
+
+        Expr = '*(*(((%s**)%s)+%s)+%s)' % (self.ElementType.GetName(), self.AllocatorData.GetValue(), index / self.ChunkBytesSize, index % self.ChunkBytesSize)
+        return self.valobj.CreateValueFromExpression('[%s]' % index, Expr)
+
+    def extract_type(self):
+        logger = lldb.formatters.Logger.Logger()
+        ArrayType = self.valobj.GetType().GetUnqualifiedType()
+        if ArrayType.IsReferenceType():
+            ArrayType = ArrayType.GetDereferencedType()
+        elif ArrayType.IsPointerType():
+            ArrayType = ArrayType.GetPointeeType()
+        if ArrayType.GetNumberOfTemplateArguments() > 0:
+            ElementType = ArrayType.GetTemplateArgumentType(0)
+        else:
+            ElementType = None
+        return ElementType
+
+    def update(self):
+        logger = lldb.formatters.Logger.Logger()
+        try:
+            self.ElementType = self.extract_type()
+            self.ElementTypeSize = self.ElementType.GetByteSize()
+            assert self.ElementTypeSize != 0
+            self.NumElements = self.valobj.GetChildMemberWithName('NumElements')
+            self.Chunks = self.valobj.GetChildMemberWithName('Chunks')
+            self.ArrayNum = self.Chunks.GetChildMemberWithName('ArrayNum')
+            self.AllocatorInstance = self.Chunks.GetChildMemberWithName('AllocatorInstance')
+            self.AllocatorData = self.AllocatorInstance.GetChildMemberWithName('Data')
+        except:
+            pass
+
+    def has_children(self):
+        return True
+
+def UEChunkedArraySummaryProvider(valobj,dict):
+    return 'size=%s' % valobj.GetNumChildren()
+
+class UESparseArraySynthProvider:
+
+    def __init__(self, valobj, dict):
+        logger = lldb.formatters.Logger.Logger()
+        self.valobj = valobj
+
+    def num_children(self):
+        logger = lldb.formatters.Logger.Logger()
+        try:
+            NumBitsVal = self.NumFreeIndices.GetValueAsSigned(0)
+            ArrayNumVal = self.Data.GetChildMemberWithName('ArrayNum').GetValueAsSigned(0)
+            return ArrayNumVal - NumBitsVal;
+        except:
+            return 0;
+
+    def get_child_index(self,name):
+        logger = lldb.formatters.Logger.Logger()
+        try:
+            return int(name.lstrip('[').rstrip(']'))
+        except:
+            return None
+
+    def get_child_at_index(self,index):
+        logger = lldb.formatters.Logger.Logger()
+        logger >> "Retrieving child %s" % index
+        if index < 0:
+            return None;
+
+        if index >= self.num_children():
+            return None;
+
+        Val = None
+        if self.SecondaryDataDataVal > 0:
+            Expr = '(bool)((((int*)%s)[%s/32] >> %s) & 1)' % (self.SecondaryDataData.GetAddress(), index, index)
+            Val = self.SecondaryDataData.CreateValueFromExpression('[%s]' % index, Expr)
+        else:
+            Expr = '(bool)((((int*)%s)[%s/32] >> %s) & 1)' % (self.InlineData.GetAddress(), index, index)
+            Val = self.InlineData.CreateValueFromExpression('[5s]' % index, Expr)
+
+        if Val.GetValueAsSigned(0) != 0:
+            offset = index * self.ElementTypeSize
+            return self.AllocatorData.CreateChildAtOffset('[%s]' % index, offset, self.ElementType)
+        else:
+            return None
+
+    def extract_type(self):
+        logger = lldb.formatters.Logger.Logger()
+        ArrayType = self.valobj.GetType().GetUnqualifiedType()
+        if ArrayType.IsReferenceType():
+            ArrayType = ArrayType.GetDereferencedType()
+        elif ArrayType.IsPointerType():
+            ArrayType = ArrayType.GetPointeeType()
+        if ArrayType.GetNumberOfTemplateArguments() > 0:
+            ElementType = ArrayType.GetTemplateArgumentType(0)
+        else:
+            ElementType = None
+        return ElementType
+
+    def update(self):
+        logger = lldb.formatters.Logger.Logger()
+        try:
+            self.ElementType = self.extract_type()
+            self.ElementTypeSize = self.ElementType.GetByteSize()
+            assert self.ElementTypeSize != 0
+            self.NumFreeIndices = self.valobj.GetChildMemberWithName('NumFreeIndices')
+            self.Data = self.valobj.GetChildMemberWithName('Data')
+            self.AllocatorInstance = self.Data.GetChildMemberWithName('AllocatorInstance')
+            self.AllocatorData = self.AllocatorInstance.GetChildMemberWithName('Data')
+            self.AllocationFlags = self.valobj.GetChildMemberWithName('AllocationFlags')
+            self.InlineData = self.AllocationFlags.GetChildMemberWithName('InlineData')
+            self.SecondaryData = self.AllocationFlags.GetChildMemberWithName('SecondaryData')
+            self.SecondaryDataData = self.SecondaryData.GetChildMemberWithName('Data')
+            self.SecondaryDataDataVal = self.SecondaryDataData.GetValueAsSigned(0)
+        except:
+            pass
+
+    def has_children(self):
+        return True
+
+def UESparseArraySummaryProvider(valobj,dict):
+    return 'size=%s' % valobj.GetNumChildren()
+
+class UEBitArraySynthProvider:
+
+    def __init__(self, valobj, dict):
+        logger = lldb.formatters.Logger.Logger()
+        self.valobj = valobj
+
+    def num_children(self):
+        logger = lldb.formatters.Logger.Logger()
+        try:
+            NumBitsVal = self.NumBits.GetValueAsSigned(0)
+            return NumBitsVal;
+        except:
+            return 0;
+
+    def get_child_index(self,name):
+        logger = lldb.formatters.Logger.Logger()
+        try:
+            return int(name.lstrip('[').rstrip(']'))
+        except:
+            return None
+
+    def get_child_at_index(self,index):
+        logger = lldb.formatters.Logger.Logger()
+        logger >> "Retrieving child %s" % index
+        if index < 0:
+            return None;
+
+        if index >= self.num_children():
+            return None;
+
+        if self.SecondaryDataDataVal > 0:
+            Expr = '(bool)((((int*)%s)[%s/32] >> %s) & 1)' % (self.SecondaryDataData.GetAddress(), index, index)
+            return self.SecondaryDataData.CreateValueFromExpression('[%s]' % index, Expr)
+        else:
+            Expr = '(bool)((((int*)%s)[%s/32] >> %s) & 1)' % (self.InlineData.GetAddress(), index, index)
+            return self.InlineData.CreateValueFromExpression('[%s]' % index, Expr)
+
+        return None
+
+    def update(self):
+        logger = lldb.formatters.Logger.Logger()
+        try:
+            self.SecondaryDataData = None
+            self.SecondaryDataDataVal = 0
+            self.NumBits = self.valobj.GetChildMemberWithName('NumBits')
+            self.MaxBits = self.valobj.GetChildMemberWithName('MaxBits')
+            self.InlineData = self.valobj.GetChildMemberWithName('InlineData')
+            self.SecondaryData = self.valobj.GetChildMemberWithName('SecondaryData')
+            if self.SecondaryData != None:
+                self.SecondaryDataData = self.SecondaryData.GetChildMemberWithName('Data')
+                self.SecondaryDataDataVal = self.SecondaryDataData.GetValueAsUnsigned(0)
+        except:
+            pass
+
+    def has_children(self):
+        return True
+
+def UEBitArraySummaryProvider(valobj,dict):
+    return 'size=%s' % valobj.GetNumChildren()
+
+class UEArraySynthProvider:
+
+    def __init__(self, valobj, dict):
+        logger = lldb.formatters.Logger.Logger()
+        self.valobj = valobj
+
+    def num_children(self):
+        logger = lldb.formatters.Logger.Logger()
+        try:
+            ArrayNumVal = self.ArrayNum.GetValueAsSigned(0)
+            return ArrayNumVal + self.NumChildren;
+        except:
+            return 0;
+
+    def get_child_index(self,name):
+        logger = lldb.formatters.Logger.Logger()
+        try:
+            return self.NumChildren + int(name.lstrip('[').rstrip(']'))
+        except:
+            return self.valobj.GetIndexOfChildWithName(name)
+
+    def get_child_at_index(self,index):
+        logger = lldb.formatters.Logger.Logger()
+        logger >> "Retrieving child %s" % index
+        if index < 0:
+            return None;
+
+        if index < self.NumChildren:
+            return self.valobj.GetChildAtIndex(index)
+        else:
+            index -= self.NumChildren
+
+        if index >= self.num_children():
+            return None;
+        try:
+            offset = index * self.ElementTypeSize
+            if self.Data != None:
+                return self.Data.CreateChildAtOffset('[%s]' % index, offset, self.ElementType)
+            elif self.SecondaryDataDataVal > 0:
+                return self.SecondaryDataData.CreateChildAtOffset('[%s]' % index, offset, self.ElementType)
+            else:
+                return self.InlineData.CreateChildAtOffset('[%s]' % index, offset, self.ElementType)
+        except:
+            return None
+
+    def extract_type(self):
+        logger = lldb.formatters.Logger.Logger()
+        ArrayType = self.valobj.GetType().GetUnqualifiedType()
+        if ArrayType.IsReferenceType():
+            ArrayType = ArrayType.GetDereferencedType()
+        elif ArrayType.IsPointerType():
+            ArrayType = ArrayType.GetPointeeType()
+
+        if ArrayType.GetNumberOfTemplateArguments() > 0:
+            ElementType = ArrayType.GetTemplateArgumentType(0)
+        else:
+            ElementType = None
+        return ElementType
+
+    def update(self):
+        logger = lldb.formatters.Logger.Logger()
+        try:
+            self.NumChildren = self.valobj.GetNumChildren()
+            self.ArrayNum = self.valobj.GetChildMemberWithName('ArrayNum')
+            self.ArrayMax = self.valobj.GetChildMemberWithName('ArrayMax')
+            self.AllocatorInstance = self.valobj.GetChildMemberWithName('AllocatorInstance')
+            if self.AllocatorInstance.GetType().IsReferenceType():
+                self.AllocatorInstance = self.AllocatorInstance.Dereference()
+            self.Data = self.AllocatorInstance.GetChildMemberWithName('Data')
+            self.InlineData = self.AllocatorInstance.GetChildMemberWithName('InlineData')
+            self.SecondaryData = self.AllocatorInstance.GetChildMemberWithName('SecondaryData')
+            if self.SecondaryData != None:
+                self.SecondaryDataData = self.SecondaryData.GetChildMemberWithName('Data')
+                self.SecondaryDataDataVal = self.SecondaryDataData.GetValueAsUnsigned(0)
+            else:
+                self.SecondaryDataData = None
+                self.SecondaryDataDataVal = 0
+        except:
+            logger >> "UEArraySynthProvider::update failed accessing members"
+            pass
+        try:
+            self.ElementType = self.extract_type()
+            self.ElementTypeSize = self.ElementType.GetByteSize()
+            assert self.ElementTypeSize != 0
+        except:
+            logger >> "UEArraySynthProvider::update failed accessing element type"
+            pass
+
+    def has_children(self):
+        return True
+
+def UEArraySummaryProvider(valobj,dict):
+    return 'size=%s' % valobj.GetChildMemberWithName('ArrayNum').GetValueAsSigned(0)
+
+class UESetSynthProvider:
+
+    def __init__(self, valobj, dict):
+        logger = lldb.formatters.Logger.Logger()
+        self.valobj = valobj
+
+        # Can't cast to TSetElement - the template instantiation won't allow it
+        Expr = '(size_t)sizeof(FSetElementId) + sizeof(int32)'
+        self.TSetElement = self.valobj.CreateValueFromExpression('TSetElement', Expr)
+
+    def num_children(self):
+        logger = lldb.formatters.Logger.Logger()
+        try:
+            ArrayNumVal = self.ArrayNum.GetValueAsUnsigned(0)
+            NumFreeIndicesVal = self.NumFreeIndices.GetValueAsUnsigned(0)
+            return ArrayNumVal - NumFreeIndicesVal;
+        except:
+            return 0;
+
+    def get_child_index(self,name):
+        logger = lldb.formatters.Logger.Logger()
+        try:
+            return int(name.lstrip('[').rstrip(']'))
+        except:
+            return 0
+
+    def get_child_at_index(self,index):
+        logger = lldb.formatters.Logger.Logger()
+        logger >> "Retrieving child %s" % index
+        if index < 0:
+            return None;
+
+        if index >= self.num_children():
+            return None;
+        try:
+            offset = index * self.ElementTypeSize
+            HasObject = 0
+            if self.SecondaryDataDataVal > 0:
+                Expr = '(bool)((((int*)%s)[%s/32] >> %s) & 1)' % (self.SecondaryDataDataVal, index, index)
+                HasObject = 1 ##self.AllocationFlagsSecondaryDataData.CreateValueFromExpression('[%s]' % index, Expr).GetValueAsUnsigned(0)
+            else:
+                Expr = '(bool)((((int*)%s)[%s/32] >> %s) & 1)' % (self.AllocationFlagsInlineDataAddr, index, index)
+                HasObject = 1 ##self.AllocationFlagsInlineData.CreateValueFromExpression('[%s]' % index, Expr).GetValueAsUnsigned(0)
+
+            if HasObject == 1:
+                return self.AllocatorInstanceData.CreateChildAtOffset('[%s]' % index, offset, self.ElementType)
+            else:
+                return self.valobj.CreateValueFromExpression('[%s]' % index, '(void*)0xDEADBEEF')
+        except:
+            return None
+
+    def extract_type(self):
+        logger = lldb.formatters.Logger.Logger()
+        ArrayType = self.valobj.GetType().GetUnqualifiedType()
+        if ArrayType.IsReferenceType():
+            ArrayType = ArrayType.GetDereferencedType()
+        elif ArrayType.IsPointerType():
+            ArrayType = ArrayType.GetPointeeType()
+        if ArrayType.GetNumberOfTemplateArguments() > 0:
+            ElementType = ArrayType.GetTemplateArgumentType(0)
+        else:
+            ElementType = None
+        return ElementType
+
+    def update(self):
+        logger = lldb.formatters.Logger.Logger()
+        try:
+            self.Elements = self.valobj.GetChildMemberWithName('Elements')
+            self.ElementsData = self.Elements.GetChildMemberWithName('Data')
+            self.ArrayNum = self.ElementsData.GetChildMemberWithName('ArrayNum')
+            self.NumFreeIndices = self.Elements.GetChildMemberWithName('NumFreeIndices')
+            self.AllocationFlags = self.Elements.GetChildMemberWithName('AllocationFlags')
+            self.AllocationFlagsAllocatorInstance = self.AllocationFlags.GetChildMemberWithName('AllocatorInstance')
+            self.AllocatorInstance = self.ElementsData.GetChildMemberWithName('AllocatorInstance')
+            self.AllocatorInstanceData = self.AllocatorInstance.GetChildMemberWithName('Data')
+            self.AllocationFlagsInlineData = self.AllocationFlagsAllocatorInstance.GetChildMemberWithName('InlineData')
+            self.AllocationFlagsInlineDataAddr = self.AllocationFlagsInlineData.AddressOf().GetValueAsUnsigned(0)
+            self.AllocationFlagsSecondaryData = self.AllocationFlagsAllocatorInstance.GetChildMemberWithName('SecondaryData')
+            self.AllocationFlagsSecondaryDataData = self.AllocationFlagsSecondaryData.GetChildMemberWithName('Data')
+            self.SecondaryDataDataVal = self.AllocationFlagsSecondaryDataData.GetValueAsUnsigned(0)
+            self.ElementType = self.extract_type()
+            # This may fail due to C++ struct padding - will have to check
+            self.ElementTypeSize = self.ElementType.GetByteSize() + self.TSetElement.GetValueAsUnsigned(0)
+            assert self.ElementTypeSize != 0
+        except:
+            pass
+
+    def has_children(self):
+        return True
+
+def UESetSummaryProvider(valobj,dict):
+    return 'size=%s' % valobj.GetNumChildren()
+
+class UEMapSynthProvider:
+
+    def __init__(self, valobj, dict):
+        logger = lldb.formatters.Logger.Logger()
+        self.valobj = valobj
+
+        # Can't cast to TSetElement - the template instantiation won't allow it
+        Expr = '(size_t)sizeof(FSetElementId) + sizeof(int32)'
+        self.TSetElement = self.valobj.CreateValueFromExpression('TSetElement', Expr)
+
+    def num_children(self):
+        logger = lldb.formatters.Logger.Logger()
+        try:
+            ArrayNumVal = self.ArrayNum.GetValueAsUnsigned(0)
+            NumFreeIndicesVal = self.NumFreeIndices.GetValueAsUnsigned(0)
+            return ArrayNumVal - NumFreeIndicesVal;
+        except:
+            return 0;
+
+    def get_child_index(self,name):
+        logger = lldb.formatters.Logger.Logger()
+        try:
+            return int(name.lstrip('[').rstrip(']'))
+        except:
+            return 0
+
+    def get_child_at_index(self,index):
+        logger = lldb.formatters.Logger.Logger()
+        logger >> "Retrieving child %s" % index
+        if index < 0:
+            return None;
+
+        if index >= self.num_children():
+            return None;
+        try:
+            offset = index * self.ElementTypeSize
+            HasObject = 0
+            if self.SecondaryDataDataVal != 0:
+                Expr = '(bool)((((unsigned int*)%s)[%s/32] >> %s) & 1)' % (self.SecondaryDataDataVal, index, index)
+                HasObject = 1 ##self.AllocationFlagsSecondaryDataData.CreateValueFromExpression('[%s]' % index, Expr).GetValueAsUnsigned(0)
+            else:
+                Expr = '(bool)((((unsigned int*)%s)[%s/32] >> %s) & 1)' % (self.AllocationFlagsInlineDataAddr, index, index)
+                HasObject = 1 ##self.AllocationFlagsInlineData.CreateValueFromExpression('[%s]' % index, Expr).GetValueAsUnsigned(0)
+
+            if HasObject == 1:
+                return self.AllocatorInstanceData.CreateChildAtOffset('[%s]' % index,offset,self.ElementType)
+            else:
+                return self.valobj.CreateValueFromExpression('[%s]' % index, '(void*)0xDEADBEEF')
+        except:
+            return None
+
+    def extract_type(self):
+        logger = lldb.formatters.Logger.Logger()
+        ArrayType = self.Pairs.GetType().GetUnqualifiedType()
+        if ArrayType.IsReferenceType():
+            ArrayType = ArrayType.GetDereferencedType()
+        elif ArrayType.IsPointerType():
+            ArrayType = ArrayType.GetPointeeType()
+        if ArrayType.GetNumberOfTemplateArguments() > 0:
+            ElementType = ArrayType.GetTemplateArgumentType(0)
+        else:
+            ElementType = None
+        return ElementType
+
+    def update(self):
+        logger = lldb.formatters.Logger.Logger()
+        try:
+            self.Pairs = self.valobj.GetChildMemberWithName('Pairs')
+            self.Elements = self.Pairs.GetChildMemberWithName('Elements')
+            self.ElementsData = self.Elements.GetChildMemberWithName('Data')
+            self.ArrayNum = self.ElementsData.GetChildMemberWithName('ArrayNum')
+            self.NumFreeIndices = self.Elements.GetChildMemberWithName('NumFreeIndices')
+            self.AllocationFlags = self.Elements.GetChildMemberWithName('AllocationFlags')
+            self.AllocationFlagsAllocatorInstance = self.AllocationFlags.GetChildMemberWithName('AllocatorInstance')
+            self.AllocatorInstance = self.ElementsData.GetChildMemberWithName('AllocatorInstance')
+            self.AllocatorInstanceData = self.AllocatorInstance.GetChildMemberWithName('Data')
+            self.AllocationFlagsInlineData = self.AllocationFlagsAllocatorInstance.GetChildMemberWithName('InlineData')
+            self.AllocationFlagsInlineDataAddr = self.AllocationFlagsInlineData.AddressOf().GetValueAsUnsigned(0)
+            self.AllocationFlagsSecondaryData = self.AllocationFlagsAllocatorInstance.GetChildMemberWithName('SecondaryData')
+            self.AllocationFlagsSecondaryDataData = self.AllocationFlagsSecondaryData.GetChildMemberWithName('Data')
+            self.SecondaryDataDataVal = self.AllocationFlagsSecondaryDataData.GetValueAsUnsigned(0)
+            self.ElementType = self.extract_type()
+            # This may fail due to C++ struct padding - will have to check
+            self.ElementTypeSize = self.ElementType.GetByteSize() + self.TSetElement.GetValueAsUnsigned(0)
+            assert self.ElementTypeSize != 0
+        except:
+            pass
+
+    def has_children(self):
+        return True
+
+def UEMapSummaryProvider(valobj,dict):
+    return 'size=%s' % valobj.GetNumChildren()
+
+def __lldb_init_module(debugger,dict):
+    debugger.HandleCommand('type summary add -F UEDataFormatters_2ByteChars.UETCharSummaryProvider -e TCHAR -w UEDataFormatters')
+    debugger.HandleCommand('type summary add -F UEDataFormatters_2ByteChars.UEFStringSummaryProvider -e -x "FString$" -w UEDataFormatters')
+    debugger.HandleCommand('type summary add -F UEDataFormatters_2ByteChars.UEFNameEntrySummaryProvider -e -x "FNameEntry$" -w UEDataFormatters')
+    debugger.HandleCommand('type summary add -F UEDataFormatters_2ByteChars.UEFNameSummaryProvider -e -x "FName$" -w UEDataFormatters')
+    debugger.HandleCommand('type summary add -F UEDataFormatters_2ByteChars.UEFNameSummaryProvider -e -x "FMinimalName$" -w UEDataFormatters')
+    debugger.HandleCommand('type summary add -F UEDataFormatters_2ByteChars.UEUObjectBaseSummaryProvider -e UObject -w UEDataFormatters')
+    debugger.HandleCommand('type summary add -F UEDataFormatters_2ByteChars.UEUObjectBaseSummaryProvider -e UObjectBase -w UEDataFormatters')
+    debugger.HandleCommand('type summary add -F UEDataFormatters_2ByteChars.UEUObjectBaseSummaryProvider -e UObjectBaseUtility -w UEDataFormatters')
+    debugger.HandleCommand('type summary add -F UEDataFormatters_2ByteChars.UEFFieldClassSummaryProvider -e FFieldClass -w UEDataFormatters')
+    debugger.HandleCommand('type summary add -F UEDataFormatters_2ByteChars.UEFFieldSummaryProvider -e FField -w UEDataFormatters')
+    debugger.HandleCommand('type summary add -F UEDataFormatters_2ByteChars.UEFWeakObjectPtrSummaryProvider -e FWeakObjectPtr -w UEDataFormatters')
+    debugger.HandleCommand('type synthetic add -l UEDataFormatters_2ByteChars.UETWeakObjectPtrSynthProvider -x "TWeakObjectPtr<.+>$" -w UEDataFormatters')
+    debugger.HandleCommand('type synthetic add -l UEDataFormatters_2ByteChars.UETWeakObjectPtrSynthProvider -x "TAutoWeakObjectPtr<.+>$" -w UEDataFormatters')
+    debugger.HandleCommand('type synthetic add -l UEDataFormatters_2ByteChars.UEArraySynthProvider -x "TArray<.+,.+>$" -w UEDataFormatters')
+    debugger.HandleCommand('type summary add -F UEDataFormatters_2ByteChars.UEArraySummaryProvider -e -x "TArray<.+>$" -w UEDataFormatters')
+    debugger.HandleCommand('type synthetic add -l UEDataFormatters_2ByteChars.UEBitArraySynthProvider -x "TBitArray<.+>$" -w UEDataFormatters')
+    debugger.HandleCommand('type summary add -F UEDataFormatters_2ByteChars.UEBitArraySummaryProvider -e -x "TBitArray<.+>$" -w UEDataFormatters')
+    debugger.HandleCommand('type synthetic add -l UEDataFormatters_2ByteChars.UESparseArraySynthProvider -x "TSparseArray<.+>$" -w UEDataFormatters')
+    debugger.HandleCommand('type summary add -F UEDataFormatters_2ByteChars.UESparseArraySummaryProvider -e -x "TSparseArray<.+>$" -w UEDataFormatters')
+    debugger.HandleCommand('type synthetic add -l UEDataFormatters_2ByteChars.UEChunkedArraySynthProvider -x "TChunkedArray<.+>$" -w UEDataFormatters')
+    debugger.HandleCommand('type summary add -F UEDataFormatters_2ByteChars.UEChunkedArraySummaryProvider -e -x "TChunkedArray<.+>$" -w UEDataFormatters')
+    debugger.HandleCommand('type synthetic add -l UEDataFormatters_2ByteChars.UESetSynthProvider -x "TSet<.+>$" -w UEDataFormatters')
+    debugger.HandleCommand('type summary add -F UEDataFormatters_2ByteChars.UESetSummaryProvider -e -x "TSet<.+>$" -w UEDataFormatters')
+    debugger.HandleCommand('type synthetic add -l UEDataFormatters_2ByteChars.UEMapSynthProvider -x "TMap<.+>$" -w UEDataFormatters')
+    debugger.HandleCommand('type summary add -F UEDataFormatters_2ByteChars.UEMapSummaryProvider -e -x "TMap<.+>$" -w UEDataFormatters')
+    debugger.HandleCommand('type synthetic add -l UEDataFormatters_2ByteChars.UEMapSynthProvider -x "TMapBase<.+>$" -w UEDataFormatters')
+    debugger.HandleCommand('type summary add -F UEDataFormatters_2ByteChars.UEMapSummaryProvider -e -x "TMapBase<.+>$" -w UEDataFormatters')
+    debugger.HandleCommand("type category enable UEDataFormatters")
+
