Fix up LLDB pretty printers.  WeakObjectPtr's were indexed incorrectly using a floating point value as the result of a division in the data formatter python scripts

diff --git a/Engine/Extras/LLDBDataFormatters/UEDataFormatters.py b/Engine/Extras/LLDBDataFormatters/UEDataFormatters.py
--- a/Engine/Extras/LLDBDataFormatters/UEDataFormatters.py
+++ b/Engine/Extras/LLDBDataFormatters/UEDataFormatters.py
@@ -100,7 +100,7 @@ def UEFNameSummaryProvider(valobj,dict):
     else:
         NameStr = UEFNameEntrySummaryProvider(NameEntry, dict)
         if Number.GetValueAsUnsigned(0) != 0:
-            return "'%s'_%d" % (NameStr, Number-1)
+            return "'%s'_%d" % (NameStr, Number.GetValueAsUnsigned(0)-1)
         else:
             return "'%s'" % NameStr
 
@@ -134,11 +134,11 @@ class UETWeakObjectPtrSynthProvider:
         logger = lldb.formatters.Logger.Logger()
         logger >> "Retrieving child " + str(index)
         if self.ObjectSerialNumberVal >= 1:
-            Expr = 'GObjectArrayForDebugVisualizers->Objects['+str(self.ObjectIndexVal/65536)+']['+str(self.ObjectIndexVal%65536)+'].SerialNumber == '+str(self.ObjectSerialNumberVal)
+            Expr = 'GObjectArrayForDebugVisualizers->Objects['+str(int(self.ObjectIndexVal/65536))+']['+str(self.ObjectIndexVal%65536)+'].SerialNumber == '+str(self.ObjectSerialNumberVal)
             Val = self.valobj.CreateValueFromExpression(str(self.ObjectIndexVal), Expr)
             Value = Val.GetValueAsUnsigned(0)
             if Value != 0:
-                Expr = 'GObjectArrayForDebugVisualizers->Objects['+str(self.ObjectIndexVal/65536)+']['+str(self.ObjectIndexVal%65536)+'].Object'
+                Expr = 'GObjectArrayForDebugVisualizers->Objects['+str(int(self.ObjectIndexVal/65536))+']['+str(self.ObjectIndexVal%65536)+'].Object'
                 return self.valobj.CreateValueFromExpression('Object', Expr)
             else:
                 Expr = '(void*)0xDEADBEEF'
@@ -167,13 +167,13 @@ def UEFWeakObjectPtrSummaryProvider(valobj,dict):
         return 'object=nullptr'
     ObjectIndex = valobj.GetChildMemberWithName('ObjectIndex')
     ObjectIndexVal = ObjectIndex.GetValueAsSigned(0)
-    Expr = 'GObjectArrayForDebugVisualizers->Objects['+str(ObjectIndexVal/65536)+']['+str(ObjectIndexVal%65536)+'].SerialNumber == '+str(ObjectSerialNumberVal)
+    Expr = 'GObjectArrayForDebugVisualizers->Objects['+str(int(ObjectIndexVal/65536))+']['+str(ObjectIndexVal%65536)+'].SerialNumber == '+str(ObjectSerialNumberVal)
     Val = valobj.CreateValueFromExpression(str(ObjectIndexVal), Expr)
     ValRef = Val.GetValueAsUnsigned(0)
     if ValRef == 0:
         return 'object=STALE'
     else:
-        Expr = 'GObjectArrayForDebugVisualizers->Objects['+str(ObjectIndexVal/65536)+']['+str(ObjectIndexVal%65536)+'].Object'
+        Expr = 'GObjectArrayForDebugVisualizers->Objects['+str(int(ObjectIndexVal/65536))+']['+str(ObjectIndexVal%65536)+'].Object'
         Val = valobj.CreateValueFromExpression(str(ObjectIndexVal), Expr)
         return 'object=' + Val.GetValue()
 
diff --git a/Engine/Extras/LLDBDataFormatters/UEDataFormatters_2ByteChars.py b/Engine/Extras/LLDBDataFormatters/UEDataFormatters_2ByteChars.py
--- a/Engine/Extras/LLDBDataFormatters/UEDataFormatters_2ByteChars.py
+++ b/Engine/Extras/LLDBDataFormatters/UEDataFormatters_2ByteChars.py
@@ -156,11 +156,11 @@ class UETWeakObjectPtrSynthProvider:
         logger = lldb.formatters.Logger.Logger()
         logger >> "Retrieving child %s" % index
         if self.ObjectSerialNumberVal >= 1:
-            Expr = 'GObjectArrayForDebugVisualizers->Objects[%s][%s].SerialNumber == %s' % (self.ObjectIndexVal/65536, self.ObjectIndexVal%65536, self.ObjectSerialNumberVal)
+            Expr = 'GObjectArrayForDebugVisualizers->Objects[%s][%s].SerialNumber == %s' % (int(self.ObjectIndexVal/65536), self.ObjectIndexVal%65536, self.ObjectSerialNumberVal)
             Val = self.valobj.CreateValueFromExpression("%s" % self.ObjectIndexVal, Expr)
             Value = Val.GetValueAsUnsigned(0)
             if Value != 0:
-                Expr = 'GObjectArrayForDebugVisualizers->Objects[%s][%s].Object' % (self.ObjectIndexVal/65536, self.ObjectIndexVal%65536)
+                Expr = 'GObjectArrayForDebugVisualizers->Objects[%s][%s].Object' % (int(self.ObjectIndexVal/65536), self.ObjectIndexVal%65536)
                 return self.valobj.CreateValueFromExpression('Object', Expr)
             else:
                 Expr = '(void*)0xDEADBEEF'
@@ -189,13 +189,13 @@ def UEFWeakObjectPtrSummaryProvider(valobj,dict):
         return 'object=nullptr'
     ObjectIndex = valobj.GetChildMemberWithName('ObjectIndex')
     ObjectIndexVal = ObjectIndex.GetValueAsSigned(0)
-    Expr = 'GObjectArrayForDebugVisualizers->Objects[%s][%s].SerialNumber == %s' % (ObjectIndexVal/65536, ObjectIndexVal%65536, ObjectSerialNumberVal)
+    Expr = 'GObjectArrayForDebugVisualizers->Objects[%s][%s].SerialNumber == %s' % (int(ObjectIndexVal/65536), ObjectIndexVal%65536, ObjectSerialNumberVal)
     Val = valobj.CreateValueFromExpression('%s' % ObjectIndexVal, Expr)
     ValRef = Val.GetValueAsUnsigned(0)
     if ValRef == 0:
         return 'object=STALE'
     else:
-        Expr = 'GObjectArrayForDebugVisualizers->Objects[%s][%s].Object' % (ObjectIndexVal/65536, ObjectIndexVal%65536)
+        Expr = 'GObjectArrayForDebugVisualizers->Objects[%s][%s].Object' % (int(ObjectIndexVal/65536), ObjectIndexVal%65536)
         Val = valobj.CreateValueFromExpression("%s" % ObjectIndexVal, Expr)
         return 'object=' + Val.GetValue()
 
